Add a stub low-level tests program for the NetworkPrediction plugin.
Migrated and updated the only test code I found in the plugin, some basic predicted cue functionality, to the test program.
Low-level test runner InitAll helper function now sets GIsRunning to true, NetworkPrediction is checking for this during initialization.

diff --git a/Engine/Build/LowLevelTests_GenProps.xml b/Engine/Build/LowLevelTests_GenProps.xml
--- a/Engine/Build/LowLevelTests_GenProps.xml
+++ b/Engine/Build/LowLevelTests_GenProps.xml
@@ -1,24 +1,31 @@
 ﻿<?xml version="1.0" encoding="utf-8"?>
 <BuildGraph xmlns:xsi="http://www.w3.org/2001/XMLSchema-instance" xsi:schemaLocation="http://www.epicgames.com/BuildGraph ../../Build/Graph/Schema.xsd" xmlns="http://www.epicgames.com/BuildGraph">
-  <Property Name="TestNames" Value="HeadlessChaos;Foundation;ReplicationSystem" />
+  <Property Name="TestNames" Value="HeadlessChaos;Foundation;ReplicationSystem;NetworkPredictionPlugin" />
+  <Property Name="NetworkPredictionPluginDisabled" Value="False" />
   <Property Name="ReplicationSystemDisabled" Value="False" />
   <Property Name="FoundationDisabled" Value="False" />
   <Property Name="HeadlessChaosDisabled" Value="False" />
+  <Property Name="NetworkPredictionPluginShort" Value="Net Prediction" />
   <Property Name="ReplicationSystemShort" Value="Replication System" />
   <Property Name="FoundationShort" Value="Foundation" />
   <Property Name="HeadlessChaosShort" Value="Headless Chaos" />
+  <Property Name="NetworkPredictionPluginTarget" Value="NetworkPredictionTests" />
   <Property Name="ReplicationSystemTarget" Value="ReplicationSystemLowLevelTests" />
   <Property Name="FoundationTarget" Value="FoundationTests" />
   <Property Name="HeadlessChaosTarget" Value="HeadlessChaos" />
+  <Property Name="NetworkPredictionPluginBinariesRelative" Value="Engine\Binaries" />
   <Property Name="ReplicationSystemBinariesRelative" Value="Engine\Binaries" />
   <Property Name="FoundationBinariesRelative" Value="Engine\Binaries" />
   <Property Name="HeadlessChaosBinariesRelative" Value="Engine\Binaries" />
+  <Property Name="NetworkPredictionPluginReportType" Value="console" />
   <Property Name="ReplicationSystemReportType" Value="console" />
   <Property Name="FoundationReportType" Value="console" />
   <Property Name="HeadlessChaosReportType" Value="console" />
+  <Option Name="RunNetworkPredictionPluginTests" DefaultValue="False" Description="Run Net Prediction Tests" />
   <Option Name="RunReplicationSystemTests" DefaultValue="False" Description="Run Replication System Tests" />
   <Option Name="RunFoundationTests" DefaultValue="False" Description="Run Foundation Tests" />
   <Option Name="RunHeadlessChaosTests" DefaultValue="False" Description="Run Headless Chaos Tests" />
+  <Property Name="NetworkPredictionPluginSupportedPlatforms" Value="Win64;Linux" />
   <Property Name="ReplicationSystemSupportedPlatforms" Value="Win64;Linux" />
   <Property Name="FoundationSupportedPlatforms" Value="Win64;Mac;Linux;Android" />
   <Property Name="HeadlessChaosSupportedPlatforms" Value="Win64" />
diff --git a/Engine/Plugins/Runtime/NetworkPrediction/Source/NetworkPrediction/Private/NetworkPredictionCues.cpp b/Engine/Plugins/Runtime/NetworkPrediction/Source/NetworkPrediction/Private/NetworkPredictionCues.cpp
--- a/Engine/Plugins/Runtime/NetworkPrediction/Source/NetworkPrediction/Private/NetworkPredictionCues.cpp
+++ b/Engine/Plugins/Runtime/NetworkPrediction/Source/NetworkPrediction/Private/NetworkPredictionCues.cpp
@@ -11,221 +11,3 @@ FGlobalCueTypeTable::FRegisteredCueTypeInfo& FGlobalCueTypeTable::GetRegistedTyp
 	static FGlobalCueTypeTable::FRegisteredCueTypeInfo PendingCueTypes;
 	return PendingCueTypes;
 }
-
-// ------------------------------------------------------------------------------------------------------------------------------------------
-//	Mock Cue example: Minimal example of NetSimCue pipeline
-//		(1)	- How to define cue types and traits
-//		(2) - How to define cue sets
-//		(3) - How to define handlers
-//		(4) - How to invoke cues in simulation
-//
-// ------------------------------------------------------------------------------------------------------------------------------------------
-
-#define ENABLE_MOCK_CUE 0
-#if ENABLE_MOCK_CUE
-
-// ---------------------------------------------------------------------
-// (1) Mock Cue Types
-// ---------------------------------------------------------------------
-
-// Example "Weak" cue (will default to NetSimCueTraits::Default traits, which are "Weak" non replicated, non resimulated). No NetSerialize/NetIdentical is required.
-struct FMockImpactCue
-{
-	NETSIMCUE_BODY();
-
-	FMockImpactCue() = default;
-	FMockImpactCue(const FVector& InImpact) : ImpactLocation(InImpact) { }
-
-	FVector ImpactLocation;
-
-};
-NETSIMCUE_REGISTER(FMockImpactCue, TEXT("Impact")); // Must go in .cpp
-static_assert(TNetSimCueTraits<FMockImpactCue>::InvokeMask == NetSimCueTraits::Default::InvokeMask, "Traits Errors"); // Not required: just asserting what this classes traits are
-
-// ----------------
-
-// Example "Strong" cue. Explicitly sets NetSimCueTraits::Strong and implements NetSerialize/NetIdentical (will not compile without these functions) 
-struct FMockDamageCue
-{
-	NETSIMCUE_BODY();
-
-	FMockDamageCue() = default;
-	FMockDamageCue(const int32 InSourceID, const int32 InDamageType, const FVector& InHitLocation) 
-		: SourceID(InSourceID), DamageType(InDamageType), HitLocation(InHitLocation) { }
-
-	int32 SourceID;
-	int32 DamageType;
-	FVector HitLocation;
-
-	using Traits = NetSimCueTraits::Strong; // Sets our traits via preset.
-	
-	void NetSerialize(FArchive& Ar)
-	{
-		Ar << SourceID;
-		Ar << DamageType;
-		Ar << HitLocation;
-	}
-
-	bool NetIdentical(const FMockDamageCue& Other) const
-	{
-		const float ErrorTolerance = 1.f;
-		return SourceID == Other.SourceID && DamageType && Other.DamageType && HitLocation.Equals(Other.HitLocation, ErrorTolerance);
-	}
-};
-NETSIMCUE_REGISTER(FMockDamageCue, TEXT("Damage")); // Must go in .cpp
-static_assert(TNetSimCueTraits<FMockDamageCue>::InvokeMask == NetSimCueTraits::Strong::InvokeMask, "Traits Errors"); // Not required: just asserting what this classes traits are
-
-// Alternative way of explicitly setting traits
-struct FMockHealingCue
-{
-	NETSIMCUE_BODY();
-	
-	FMockHealingCue() = default;
-	FMockHealingCue(const int32 InSourceID, const int32 InHealingAmount) 
-		: SourceID(InSourceID), HealingAmount(InHealingAmount) { }
-
-	int32 SourceID;
-	int32 HealingAmount;
-
-	using Traits = TNetSimCueTraitsExplicit<ENetSimCueInvoker::All, ENetSimCueReplicationTarget::All, true>; // Sets our traits explicitly
-	
-	void NetSerialize(FArchive& Ar)
-	{
-		Ar << SourceID;
-		Ar << HealingAmount;
-	}
-
-	bool NetIdentical(const FMockHealingCue& Other) const
-	{
-		const float ErrorTolerance = 1.f;
-		return SourceID == Other.SourceID && HealingAmount == Other.HealingAmount;
-	}
-};
-NETSIMCUE_REGISTER(FMockHealingCue, TEXT("Healing")); // Must go in .cpp
-static_assert(TNetSimCueTraits<FMockHealingCue>::InvokeMask == ENetSimCueInvoker::All, "Traits Errors"); // Not required: just asserting what this classes traits are
-
-// ---------------------------------------------------------------------
-// (2) Cue Sets: just a collection of cue types that a simulation emits.
-//
-//	This isn't really a core type, it is just a pattern for grouping cue types
-//	together so that handlers can register with sets of cues rather than each cue individually
-//
-//	This example illustrates how a base set can be built off of by child simulations
-// ---------------------------------------------------------------------
-
-struct FMockCueSet_Base
-{
-	template<typename TDispatchTable>
-	static void RegisterNetSimCueTypes(TDispatchTable& DispatchTable)
-	{
-		DispatchTable.template RegisterType<FMockImpactCue>();
-	}
-};
-
-struct FMockCueSet_Child
-{
-	template<typename TDispatchTable>
-	static void RegisterNetSimCueTypes(TDispatchTable& DispatchTable)
-	{
-		FMockCueSet_Base::RegisterNetSimCueTypes(DispatchTable); // Superset of the base cueset
-		DispatchTable.template RegisterType<FMockDamageCue>();
-		DispatchTable.template RegisterType<FMockHealingCue>();
-	}
-};
-
-
-// ---------------------------------------------------------------------
-// (3) Mock Handlers 
-//	-Implements ::HandeCue function to receive the cue (thats it!)
-//	-Meant to be implemented by owning actor/component that is driving the NetworkedSimulation
-//	-E.g, in practice, these aren't standalone classes. These are just functions/body that you put on your actor/component.
-//
-//	Again, we show how a base handler can be implemented and how a child class can extend off of that.
-// ---------------------------------------------------------------------
-
-class TestHandlerBase
-{
-public:
-	void HandleCue(const FMockImpactCue& ImpactCue, const FNetSimCueSystemParamemters& SystemParameters)
-	{
-		UE_LOG(LogNetworkSim, Warning, TEXT("Impact!"));
-	}
-};
-
-NETSIMCUESET_REGISTER(TestHandlerBase, FMockCueSet_Base); // Must go in cpp. Says "TestHandlerBase" can accepts cues in "FMockCueSet_Base"
-
-class TestHandlerChild : public TestHandlerBase
-{
-public:
-
-	using TestHandlerBase::HandleCue; // Required for us to "use" the HandleCue methods in our parent class
-
-	void HandleCue(const FMockDamageCue& DamageCue, const FNetSimCueSystemParamemters& SystemParameters)
-	{
-		UE_LOG(LogNetworkSim, Warning, TEXT("HandleCue: Damage. SourceID: %d DamageType: %d HitLocation: %s"), DamageCue.SourceID, DamageCue.DamageType, *DamageCue.HitLocation.ToString());
-	}
-
-	void HandleCue(const FMockHealingCue& HealingCue, const FNetSimCueSystemParamemters& SystemParameters)
-	{
-		UE_LOG(LogNetworkSim, Warning, TEXT("HandleCue: Healing. SourceID: %d Healing Amount: %d"), HealingCue.SourceID, HealingCue.HealingAmount);
-	}
-};
-
-NETSIMCUESET_REGISTER(TestHandlerChild, FMockCueSet_Child); // Must go in cpp. Says "TestHandlerBase" can accepts cues in "FMockCueSet_Base" // Must go in cpp. Says "TestHandlerBase" can accepts cues in "FMockCueSet_Base"
-
-void TestCues()
-{
-	// -------------------------------------
-	// (4) Simuation/Invoke
-	// -------------------------------------
-
-	struct FBaseSimulation
-	{
-		static void TickSimulation(FNetSimCueDispatcher& Dispatcher)
-		{
-			Dispatcher.Invoke<FMockImpactCue>(FVector(1.f, 2.f, 3.f)); // Constructing emplace avoids copy and moves. This is the ideal way to invoke!
-		}
-	};
-
-	struct FChildSimulation
-	{
-		static void TickSimulation(FNetSimCueDispatcher& Dispatcher)
-		{
-			FBaseSimulation::TickSimulation(Dispatcher);
-
-			Dispatcher.Invoke<FMockDamageCue>(1, 2, FVector(1.f, 2.f, 3.f));
-			Dispatcher.Invoke<FMockHealingCue>(10, 32);
-		}
-	};
-
-	TNetSimCueDispatcher<void> ServerDispatcher;
-	FChildSimulation::TickSimulation(ServerDispatcher);
-
-	// -------------------------------------
-	// Send
-	// -------------------------------------
-
-	FNetBitWriter TempWriter(1024 << 3);
-	ServerDispatcher.NetSendSavedCues(TempWriter, ENetSimCueReplicationTarget::All, true);
-
-	// -------------------------------------
-	// Receive
-	// -------------------------------------
-
-	FNetBitReader TempReader(nullptr, TempWriter.GetData(), TempWriter.GetNumBits());
-	TNetSimCueDispatcher<void> ClientDispatcher;
-	ClientDispatcher.SetReceiveReplicationTarget(ENetSimCueReplicationTarget::All);
-	ClientDispatcher.NetRecvSavedCues(TempReader, true, 0, 0);
-
-	TestHandlerChild MyObject;
-
-	ClientDispatcher.DispatchCueRecord(MyObject, FNetworkSimTime());
-}
-
-FAutoConsoleCommandWithWorldAndArgs NetworkSimulationModelCueCmd(TEXT("nms.CueTest"), TEXT(""),
-	FConsoleCommandWithWorldAndArgsDelegate::CreateLambda([](const TArray< FString >& Args, UWorld* World) 
-{
-	TestCues();
-}));
-
-#endif // ENABLE_MOCK_CUE
\ No newline at end of file
diff --git a/Engine/Plugins/Runtime/NetworkPrediction/Source/NetworkPrediction/Public/NetworkPredictionCueTraits.h b/Engine/Plugins/Runtime/NetworkPrediction/Source/NetworkPrediction/Public/NetworkPredictionCueTraits.h
--- a/Engine/Plugins/Runtime/NetworkPrediction/Source/NetworkPrediction/Public/NetworkPredictionCueTraits.h
+++ b/Engine/Plugins/Runtime/NetworkPrediction/Source/NetworkPrediction/Public/NetworkPredictionCueTraits.h
@@ -7,7 +7,7 @@ class FArchive;
 
 // ----------------------------------------------------------------------------------------------------------------------------------------------
 //	NetSimCue Traits: compile time settings for NetSimeCue types that determine who can invoke the event and who it replicates to.
-//	See "Mock Cue Example" in NetworkedSimulationModelCues.cpp for minimal examples for using/setting traits.
+//	See "Mock Cue Example" in NetworkPredictionCueTests.cpp for minimal examples for using/setting traits.
 //
 //	There are three traits:
 //
diff --git a/Engine/Plugins/Runtime/NetworkPrediction/Source/NetworkPrediction/Public/NetworkPredictionCues.h b/Engine/Plugins/Runtime/NetworkPrediction/Source/NetworkPrediction/Public/NetworkPredictionCues.h
--- a/Engine/Plugins/Runtime/NetworkPrediction/Source/NetworkPrediction/Public/NetworkPredictionCues.h
+++ b/Engine/Plugins/Runtime/NetworkPrediction/Source/NetworkPrediction/Public/NetworkPredictionCues.h
@@ -31,7 +31,7 @@ Notes on "simulation affecting events" (E.g, *not* NetSimCues)
 -In these cases I would recommend: A) on the handler side, don't write to any non-networked-sim state if non authority or B) just don't handle the event on non authority. (Expect corrections)
 -If the handler side doesn't have state hazards, say a teleporting volume that always does the same thing: there is no reason everyone can't run the event. Its an extension of the simulation.
 
-See "Mock Cue Example" in NetworkedSimulationModelCues.cpp for minimal example of implementing the Cue types and Handler classes.
+See "Mock Cue Example" in NetworkPredictionCueTests.cpp for minimal example of implementing the Cue types and Handler classes.
 
 =============================================================================*/
 
diff --git a/Engine/Source/Developer/LowLevelTestsRunner/Private/TestCommon/Initialization.cpp b/Engine/Source/Developer/LowLevelTestsRunner/Private/TestCommon/Initialization.cpp
--- a/Engine/Source/Developer/LowLevelTestsRunner/Private/TestCommon/Initialization.cpp
+++ b/Engine/Source/Developer/LowLevelTestsRunner/Private/TestCommon/Initialization.cpp
@@ -72,6 +72,7 @@ void InitAll(bool bAllowLogging, bool bMultithreaded)
 #if WITH_COREUOBJECT
 	InitCoreUObject();
 #endif
+	GIsRunning = true;
 }
 
 void CleanupAll()
diff --git a/Engine/Source/Programs/NetworkPredictionTests/NetworkPredictionTests.Build.cs b/Engine/Source/Programs/NetworkPredictionTests/NetworkPredictionTests.Build.cs
new file mode 100644
--- /dev/null
+++ b/Engine/Source/Programs/NetworkPredictionTests/NetworkPredictionTests.Build.cs
@@ -0,0 +1,21 @@
+// Copyright Epic Games, Inc. All Rights Reserved.
+
+using UnrealBuildTool;
+
+[SupportedPlatforms("Win64", "Linux")]
+public class NetworkPredictionTests : TestModuleRules
+{
+	public NetworkPredictionTests(ReadOnlyTargetRules Target) : base(Target)
+	{
+		PrivateDependencyModuleNames.AddRange(
+			new string[]
+			{
+				"Core",
+				"CoreUObject",
+				"NetworkPrediction",
+			}
+		);
+
+		UpdateBuildGraphPropertiesFile(new Metadata("NetworkPredictionPlugin", "Net Prediction"));
+	}
+}
diff --git a/Engine/Source/Programs/NetworkPredictionTests/NetworkPredictionTests.Target.cs b/Engine/Source/Programs/NetworkPredictionTests/NetworkPredictionTests.Target.cs
new file mode 100644
--- /dev/null
+++ b/Engine/Source/Programs/NetworkPredictionTests/NetworkPredictionTests.Target.cs
@@ -0,0 +1,17 @@
+// Copyright Epic Games, Inc. All Rights Reserved.
+
+using UnrealBuildTool;
+
+public class NetworkPredictionTestsTarget : TestTargetRules
+{
+	public NetworkPredictionTestsTarget(TargetInfo Target) : base(Target)
+	{
+		bCompileAgainstEngine = true;
+		bUsesSlate = false;
+
+		bUsePlatformFileStub = true;
+		bMockEngineDefaults = true;
+
+		IncludeOrderVersion = EngineIncludeOrderVersion.Latest;
+	}
+}
diff --git a/Engine/Source/Programs/NetworkPredictionTests/Private/NetworkPredictionCueTests.cpp b/Engine/Source/Programs/NetworkPredictionTests/Private/NetworkPredictionCueTests.cpp
new file mode 100644
--- /dev/null
+++ b/Engine/Source/Programs/NetworkPredictionTests/Private/NetworkPredictionCueTests.cpp
@@ -0,0 +1,281 @@
+// Copyright Epic Games, Inc. All Rights Reserved.
+
+#include "NetworkPredictionCues.h"
+#include "NetworkPredictionModelDef.h"
+#include "NetworkPredictionDriver.h"
+#include "Math/Vector.h"
+#include "Misc/Optional.h"
+#include "UObject/CoreNet.h"
+
+#include <catch2/catch_test_macros.hpp>
+
+namespace UE::Net::Private
+{
+
+// ------------------------------------------------------------------------------------------------------------------------------------------
+//	Mock Cue example: Minimal example of NetSimCue pipeline
+//		(1)	- How to define cue types and traits
+//		(2) - How to define cue sets
+//		(3) - How to define handlers
+//		(4) - How to invoke cues in simulation
+//
+// ------------------------------------------------------------------------------------------------------------------------------------------
+
+// ---------------------------------------------------------------------
+// (1) Mock Cue Types
+// ---------------------------------------------------------------------
+
+// Example "Weak" cue (will default to NetSimCueTraits::Default traits, which are "Weak" non replicated, non resimulated). No NetSerialize/NetIdentical is required.
+struct FMockImpactCue
+{
+	NETSIMCUE_BODY();
+
+	FMockImpactCue() = default;
+	FMockImpactCue(const FVector& InImpact) : ImpactLocation(InImpact) { }
+
+	FVector ImpactLocation;
+
+};
+NETSIMCUE_REGISTER(FMockImpactCue, TEXT("Impact")); // Must go in .cpp
+static_assert(TNetSimCueTraits<FMockImpactCue>::InvokeMask == NetSimCueTraits::Default::InvokeMask, "Traits Errors"); // Not required: just asserting what this classes traits are
+
+// ----------------
+
+// Example "Strong" cue. Explicitly sets NetSimCueTraits::Strong and implements NetSerialize/NetIdentical (will not compile without these functions) 
+struct FMockDamageCue
+{
+	NETSIMCUE_BODY();
+
+	FMockDamageCue() = default;
+	FMockDamageCue(const int32 InSourceID, const int32 InDamageType, const FVector& InHitLocation) 
+		: SourceID(InSourceID), DamageType(InDamageType), HitLocation(InHitLocation) { }
+
+	int32 SourceID;
+	int32 DamageType;
+	FVector HitLocation;
+
+	using Traits = NetSimCueTraits::Strong; // Sets our traits via preset.
+	
+	void NetSerialize(FArchive& Ar)
+	{
+		Ar << SourceID;
+		Ar << DamageType;
+		Ar << HitLocation;
+	}
+
+	bool NetIdentical(const FMockDamageCue& Other) const
+	{
+		const float ErrorTolerance = 1.f;
+		return SourceID == Other.SourceID && DamageType && Other.DamageType && HitLocation.Equals(Other.HitLocation, ErrorTolerance);
+	}
+};
+NETSIMCUE_REGISTER(FMockDamageCue, TEXT("Damage")); // Must go in .cpp
+static_assert(TNetSimCueTraits<FMockDamageCue>::InvokeMask == NetSimCueTraits::Strong::InvokeMask, "Traits Errors"); // Not required: just asserting what this classes traits are
+
+// Alternative way of explicitly setting traits
+struct FMockHealingCue
+{
+	NETSIMCUE_BODY();
+	
+	FMockHealingCue() = default;
+	FMockHealingCue(const int32 InSourceID, const int32 InHealingAmount) 
+		: SourceID(InSourceID), HealingAmount(InHealingAmount) { }
+
+	int32 SourceID;
+	int32 HealingAmount;
+
+	using Traits = TNetSimCueTraitsExplicit<ENetSimCueInvoker::All, ENetSimCueReplicationTarget::All, true>; // Sets our traits explicitly
+	
+	void NetSerialize(FArchive& Ar)
+	{
+		Ar << SourceID;
+		Ar << HealingAmount;
+	}
+
+	bool NetIdentical(const FMockHealingCue& Other) const
+	{
+		const float ErrorTolerance = 1.f;
+		return SourceID == Other.SourceID && HealingAmount == Other.HealingAmount;
+	}
+};
+NETSIMCUE_REGISTER(FMockHealingCue, TEXT("Healing")); // Must go in .cpp
+static_assert(TNetSimCueTraits<FMockHealingCue>::InvokeMask == ENetSimCueInvoker::All, "Traits Errors"); // Not required: just asserting what this classes traits are
+
+// ---------------------------------------------------------------------
+// (2) Cue Sets: just a collection of cue types that a simulation emits.
+//
+//	This isn't really a core type, it is just a pattern for grouping cue types
+//	together so that handlers can register with sets of cues rather than each cue individually
+//
+//	This example illustrates how a base set can be built off of by child simulations
+// ---------------------------------------------------------------------
+
+struct FMockCueSet_Base
+{
+	template<typename TDispatchTable>
+	static void RegisterNetSimCueTypes(TDispatchTable& DispatchTable)
+	{
+		DispatchTable.template RegisterType<FMockImpactCue>();
+	}
+};
+
+struct FMockCueSet_Child
+{
+	template<typename TDispatchTable>
+	static void RegisterNetSimCueTypes(TDispatchTable& DispatchTable)
+	{
+		FMockCueSet_Base::RegisterNetSimCueTypes(DispatchTable); // Superset of the base cueset
+		DispatchTable.template RegisterType<FMockDamageCue>();
+		DispatchTable.template RegisterType<FMockHealingCue>();
+	}
+};
+
+
+// ---------------------------------------------------------------------
+// (3) Mock Handlers 
+//	-Implements ::HandeCue function to receive the cue (thats it!)
+//	-Meant to be implemented by owning actor/component that is driving the NetworkedSimulation
+//	-E.g, in practice, these aren't standalone classes. These are just functions/body that you put on your actor/component.
+//
+//	Again, we show how a base handler can be implemented and how a child class can extend off of that.
+// ---------------------------------------------------------------------
+
+class TestHandlerBase
+{
+public:
+	// Store handled cue so the test code can verify it
+	TOptional<FMockImpactCue> HandledImpact;
+
+	void HandleCue(const FMockImpactCue& ImpactCue, const FNetSimCueSystemParamemters& SystemParameters)
+	{
+		HandledImpact = ImpactCue;
+	}
+};
+
+NETSIMCUESET_REGISTER(TestHandlerBase, FMockCueSet_Base); // Must go in cpp. Says "TestHandlerBase" can accepts cues in "FMockCueSet_Base"
+
+class TestHandlerChild : public TestHandlerBase
+{
+public:
+
+	using TestHandlerBase::HandleCue; // Required for us to "use" the HandleCue methods in our parent class
+
+	// Store handled cues so the test code can verify them
+	TOptional<FMockDamageCue> HandledDamage;
+	TOptional<FMockHealingCue> HandledHealing;
+
+	void HandleCue(const FMockDamageCue& DamageCue, const FNetSimCueSystemParamemters& SystemParameters)
+	{
+		HandledDamage = DamageCue;
+	}
+
+	void HandleCue(const FMockHealingCue& HealingCue, const FNetSimCueSystemParamemters& SystemParameters)
+	{
+		HandledHealing = HealingCue;
+	}
+};
+
+NETSIMCUESET_REGISTER(TestHandlerChild, FMockCueSet_Child); // Must go in cpp. Says "TestHandlerBase" can accepts cues in "FMockCueSet_Base" // Must go in cpp. Says "TestHandlerBase" can accepts cues in "FMockCueSet_Base"
+
+// -------------------------------------
+// (4) Simuation/Invoke
+// -------------------------------------
+const FVector ImpactLocation{1.f, 2.f, 3.f};
+
+const int32 DamageSourceID = 1;
+const int32 DamageType = 2;
+const FVector DamageHitLocation{4.f, 5.f, 6.f};
+	
+const int32 HealingSourceID = 10;
+const int32 HealingAmount = 32;
+
+struct FBaseSimulation
+{
+	static void TickSimulation(FNetSimCueDispatcher& Dispatcher)
+	{
+		Dispatcher.Invoke<FMockImpactCue>(ImpactLocation);
+	}
+};
+
+struct FChildSimulation
+{
+	static void TickSimulation(FNetSimCueDispatcher& Dispatcher)
+	{
+		FBaseSimulation::TickSimulation(Dispatcher);
+
+		Dispatcher.Invoke<FMockDamageCue>(DamageSourceID, DamageType, DamageHitLocation);
+		Dispatcher.Invoke<FMockHealingCue>(HealingSourceID, HealingAmount);
+	}
+};
+
+TEST_CASE("NetworkPrediction basic cue functionality")
+{
+	const int32 SimFrameNum = 1;
+	const int32 SimFrameTimeMS = 16;
+
+	TestHandlerChild ServerHandlerObject;
+
+	TNetSimCueDispatcher<FNetworkPredictionModelDef> ServerDispatcher;
+
+	ServerDispatcher.PushContext({SimFrameNum, SimFrameTimeMS, ESimulationTickContext::Authority});
+	FChildSimulation::TickSimulation(ServerDispatcher);
+
+	// No cues should have been handled on the server yet
+	REQUIRE_FALSE(ServerHandlerObject.HandledImpact.IsSet());
+	REQUIRE_FALSE(ServerHandlerObject.HandledDamage.IsSet());
+	REQUIRE_FALSE(ServerHandlerObject.HandledHealing.IsSet());
+
+	ServerDispatcher.DispatchCueRecord(ServerHandlerObject, SimFrameNum, SimFrameTimeMS, SimFrameTimeMS);
+	ServerDispatcher.PopContext();
+
+	// Server should have handled all three cues
+	REQUIRE(ServerHandlerObject.HandledImpact.IsSet());
+	REQUIRE(ServerHandlerObject.HandledImpact->ImpactLocation.Equals(ImpactLocation, 0.1f));
+
+	REQUIRE(ServerHandlerObject.HandledDamage.IsSet());
+	REQUIRE(ServerHandlerObject.HandledDamage->SourceID == DamageSourceID);
+	REQUIRE(ServerHandlerObject.HandledDamage->DamageType == DamageType);
+	REQUIRE(ServerHandlerObject.HandledDamage->HitLocation.Equals(DamageHitLocation, 0.1f));
+
+	REQUIRE(ServerHandlerObject.HandledHealing.IsSet());
+	REQUIRE(ServerHandlerObject.HandledHealing->SourceID == HealingSourceID);
+	REQUIRE(ServerHandlerObject.HandledHealing->HealingAmount == HealingAmount);
+	
+	// -------------------------------------
+	// Send
+	// -------------------------------------
+
+	FNetBitWriter TempWriter(1024 << 3);
+	ServerDispatcher.NetSendSavedCues(TempWriter, ENetSimCueReplicationTarget::All, true);
+
+	// -------------------------------------
+	// Receive
+	// -------------------------------------
+	TestHandlerChild ClientHandlerObject;
+
+	FNetBitReader TempReader(nullptr, TempWriter.GetData(), TempWriter.GetNumBits());
+	TNetSimCueDispatcher<FNetworkPredictionModelDef> ClientDispatcher;
+	ClientDispatcher.SetReceiveReplicationTarget(ENetSimCueReplicationTarget::All);
+	ClientDispatcher.NetRecvSavedCues(TempReader, true, 0, 0);
+
+	// No cues should have been handled on the client yet
+	REQUIRE_FALSE(ClientHandlerObject.HandledImpact.IsSet());
+	REQUIRE_FALSE(ClientHandlerObject.HandledDamage.IsSet());
+	REQUIRE_FALSE(ClientHandlerObject.HandledHealing.IsSet());
+
+	ClientDispatcher.DispatchCueRecord(ClientHandlerObject, SimFrameNum, SimFrameTimeMS, SimFrameTimeMS);
+
+	// Client should have handled only damage & healing, impact is not replicated
+	REQUIRE_FALSE(ClientHandlerObject.HandledImpact.IsSet());
+
+	REQUIRE(ClientHandlerObject.HandledDamage.IsSet());
+	REQUIRE(ClientHandlerObject.HandledDamage->SourceID == DamageSourceID);
+	REQUIRE(ClientHandlerObject.HandledDamage->DamageType == DamageType);
+	REQUIRE(ClientHandlerObject.HandledDamage->HitLocation.Equals(DamageHitLocation, 0.1f));
+
+	REQUIRE(ClientHandlerObject.HandledHealing.IsSet());
+	REQUIRE(ClientHandlerObject.HandledHealing->SourceID == HealingSourceID);
+	REQUIRE(ClientHandlerObject.HandledHealing->HealingAmount == HealingAmount);
+}
+
+}
\ No newline at end of file
diff --git a/Engine/Source/Programs/NetworkPredictionTests/Private/NetworkPredictionTests.cpp b/Engine/Source/Programs/NetworkPredictionTests/Private/NetworkPredictionTests.cpp
new file mode 100644
--- /dev/null
+++ b/Engine/Source/Programs/NetworkPredictionTests/Private/NetworkPredictionTests.cpp
@@ -0,0 +1,28 @@
+// Copyright Epic Games, Inc. All Rights Reserved.
+
+
+#include "Logging/LogScopedVerbosityOverride.h"
+#include "Modules/ModuleManager.h"
+#include "SlateGlobals.h"
+#include "Styling/SlateWidgetStyleContainerBase.h"
+#include "TestCommon/Initialization.h"
+
+#include <catch2/catch_test_macros.hpp>
+
+GROUP_BEFORE_GLOBAL(Catch::DefaultGroup)
+{
+	{
+		// Silence some warnings during initialization unrelated to the tests
+		LOG_SCOPE_VERBOSITY_OVERRIDE(LogSlate, ELogVerbosity::Error);
+		LOG_SCOPE_VERBOSITY_OVERRIDE(LogSlateStyle, ELogVerbosity::Error);
+		LOG_SCOPE_VERBOSITY_OVERRIDE(LogUObjectGlobals, ELogVerbosity::Error);
+		InitAll(true, true);
+	}
+
+	FModuleManager::Get().LoadModule(TEXT("NetworkPrediction"));
+}
+
+GROUP_AFTER_GLOBAL(Catch::DefaultGroup)
+{
+	CleanupAll();
+}
