Added NetworkPrediction plugin unit tests and utilities covering basic simulation ticking in forward prediction and interpolation modes.
Added some APIs to the plugin to support unit testing without needing a UWorld, actors, etc.

diff --git a/Engine/Plugins/Runtime/NetworkPrediction/Source/NetworkPrediction/Private/NetworkPredictionTrace.cpp b/Engine/Plugins/Runtime/NetworkPrediction/Source/NetworkPrediction/Private/NetworkPredictionTrace.cpp
--- a/Engine/Plugins/Runtime/NetworkPrediction/Source/NetworkPrediction/Private/NetworkPredictionTrace.cpp
+++ b/Engine/Plugins/Runtime/NetworkPrediction/Source/NetworkPrediction/Private/NetworkPredictionTrace.cpp
@@ -352,7 +352,7 @@ void FNetworkPredictionTrace::TraceSystemFault(const TCHAR* Fmt, ...)
 
 		check(Result >= 0 );
 
-	UE_LOG(LogNetworkPrediction, Warning, TEXT("SystemFault: %s"), Buffer);
+	UE_LOG(LogNetworkPrediction, Log, TEXT("SystemFault: %s"), Buffer);
 	);
 
 	UE_TRACE_LOG(NetworkPrediction, SystemFault, NetworkPredictionChannel)
diff --git a/Engine/Plugins/Runtime/NetworkPrediction/Source/NetworkPrediction/Private/NetworkPredictionWorldManager.cpp b/Engine/Plugins/Runtime/NetworkPrediction/Source/NetworkPrediction/Private/NetworkPredictionWorldManager.cpp
--- a/Engine/Plugins/Runtime/NetworkPrediction/Source/NetworkPrediction/Private/NetworkPredictionWorldManager.cpp
+++ b/Engine/Plugins/Runtime/NetworkPrediction/Source/NetworkPrediction/Private/NetworkPredictionWorldManager.cpp
@@ -81,11 +81,7 @@ void UNetworkPredictionWorldManager::OnWorldPreTick(UWorld* InWorld, ELevelTick
 
 	UE_NP_TRACE_WORLD_FRAME_START(InWorld->GetGameInstance(), InDeltaSeconds);
 
-	// Update fixed tick rate, this can be changed via editor settings
-	FixedTickState.FixedStepRealTimeMS = (1.f /  GEngine->FixedFrameRate) * 1000.f;
-	FixedTickState.FixedStepMS = (int32)FixedTickState.FixedStepRealTimeMS;
-
-	ActiveInstance = this;
+	OnWorldPreTick_Internal(InDeltaSeconds, GEngine->FixedFrameRate);
 
 	// Instantiate replicated manager on server
 	if (!ReplicatedManager && InWorld->GetNetMode() != NM_Client)
@@ -95,6 +91,15 @@ void UNetworkPredictionWorldManager::OnWorldPreTick(UWorld* InWorld, ELevelTick
 	}
 }
 
+void UNetworkPredictionWorldManager::OnWorldPreTick_Internal(float InDeltaSeconds, float InFixedFrameRate)
+{
+	// Update fixed tick rate, this can be changed via editor settings
+	FixedTickState.FixedStepRealTimeMS = (1.f /  InFixedFrameRate) * 1000.f;
+	FixedTickState.FixedStepMS = (int32)FixedTickState.FixedStepRealTimeMS;
+
+	ActiveInstance = this;
+}
+
 void UNetworkPredictionWorldManager::ReconcileSimulationsPostNetworkUpdate()
 {
 	QUICK_SCOPE_CYCLE_COUNTER(STAT_NP_RECONCILE);
@@ -105,6 +110,11 @@ void UNetworkPredictionWorldManager::ReconcileSimulationsPostNetworkUpdate()
 		return;
 	}
 
+	ReconcileSimulationsPostNetworkUpdate_Internal();
+}
+
+void UNetworkPredictionWorldManager::ReconcileSimulationsPostNetworkUpdate_Internal()
+{
 	ActiveInstance = this;
 	bLockServices = true;
 
@@ -244,7 +254,7 @@ void UNetworkPredictionWorldManager::ReconcileSimulationsPostNetworkUpdate()
 		else if (RollbackFrame > FixedTickState.PendingFrame)
 		{
 			// Most likely we haven't had a confirmed frame yet so our local frame -> server mapping hasn't been set yet
-			UE_LOG(LogNetworkPrediction, Warning, TEXT("RollbackFrame %d AHEAD of PendingFrame %d... Offset: %d"), RollbackFrame, FixedTickState.PendingFrame, FixedTickState.Offset);
+			UE_LOG(LogNetworkPrediction, Log, TEXT("RollbackFrame %d AHEAD of PendingFrame %d... Offset: %d"), RollbackFrame, FixedTickState.PendingFrame, FixedTickState.Offset);
 		}
 	}
 
@@ -268,6 +278,11 @@ void UNetworkPredictionWorldManager::BeginNewSimulationFrame(UWorld* InWorld, EL
 		return;
 	}
 
+	BeginNewSimulationFrame_Internal(DeltaTimeSeconds);
+}
+
+void UNetworkPredictionWorldManager::BeginNewSimulationFrame_Internal(float DeltaTimeSeconds)
+{
 	ActiveInstance = this;
 	bLockServices = true;
 
@@ -572,8 +587,8 @@ void UNetworkPredictionWorldManager::AdvancePhysicsResimFrame(int32& PhysicsFram
 
 bool UNetworkPredictionWorldManager::CanPhysicsFixTick() const
 {
-	npCheckSlow(GEngine);
-	return (GEngine->bUseFixedFrameRate) || Settings.bForceEngineFixTickForcePhysics;
+	// GEngine may be null in some contexts, like unit tests.
+	return (GEngine && GEngine->bUseFixedFrameRate) || Settings.bForceEngineFixTickForcePhysics;
 }
 
 void UNetworkPredictionWorldManager::SetUsingPhysics()
diff --git a/Engine/Plugins/Runtime/NetworkPrediction/Source/NetworkPrediction/Public/NetworkPredictionProxy.h b/Engine/Plugins/Runtime/NetworkPrediction/Source/NetworkPrediction/Public/NetworkPredictionProxy.h
--- a/Engine/Plugins/Runtime/NetworkPrediction/Source/NetworkPrediction/Public/NetworkPredictionProxy.h
+++ b/Engine/Plugins/Runtime/NetworkPrediction/Source/NetworkPrediction/Public/NetworkPredictionProxy.h
@@ -5,6 +5,7 @@
 #include "NetworkPredictionID.h"
 #include "NetworkPredictionStateView.h"
 #include "NetworkPredictionConfig.h"
+#include "Engine/EngineBaseTypes.h"
 #include "Engine/EngineTypes.h"
 #include "NetworkPredictionProxy.generated.h"
 
@@ -36,10 +37,23 @@ struct FNetworkPredictionProxy
 {
 	GENERATED_BODY();
 
+	template<typename ModelDef>
+	struct FInitParams
+	{
+		UNetworkPredictionWorldManager* WorldManager;
+		ENetMode Mode;
+		const FReplicationProxySet& RepProxies;
+		typename ModelDef::Simulation* Simulation = nullptr;
+		typename ModelDef::Driver* Driver = nullptr;
+	};
+
 	// The init function that you need to call. This is defined in NetworkPredictionProxyInit.h (which should only be included by your .cpp file)
 	template<typename ModelDef>
 	void Init(UWorld* World, const FReplicationProxySet& RepProxies, typename ModelDef::Simulation* Simulation=nullptr, typename ModelDef::Driver* Driver=nullptr);
 
+	template<typename ModelDef>
+	void Init(const FInitParams<ModelDef>& Params);
+
 	// When network role changes, initializes role storage and logic controller
 	void InitForNetworkRole(ENetRole Role, bool bHasNetConnection)
 	{
diff --git a/Engine/Plugins/Runtime/NetworkPrediction/Source/NetworkPrediction/Public/NetworkPredictionProxyInit.h b/Engine/Plugins/Runtime/NetworkPrediction/Source/NetworkPrediction/Public/NetworkPredictionProxyInit.h
--- a/Engine/Plugins/Runtime/NetworkPrediction/Source/NetworkPrediction/Public/NetworkPredictionProxyInit.h
+++ b/Engine/Plugins/Runtime/NetworkPrediction/Source/NetworkPrediction/Public/NetworkPredictionProxyInit.h
@@ -11,15 +11,15 @@ struct FNetworkPredictionSettings;
 // The init function binds to the templated methods on UNetworkPredictionMAnager. This will "bring in" all the templated systems on NP, so this file should only be 
 // included in your .cpp file that is calling Init.
 template<typename ModelDef>
-void FNetworkPredictionProxy::Init(UWorld* World, const FReplicationProxySet& RepProxies, typename ModelDef::Simulation* Simulation, typename ModelDef::Driver* Driver)
+void FNetworkPredictionProxy::Init(const FInitParams<ModelDef>& Params)
 {
 	using StateTypes = typename ModelDef::StateTypes;
 	using InputType = typename StateTypes::InputType;
 	using SyncType = typename StateTypes::SyncType;
 	using AuxType = typename StateTypes::AuxType;
 
-	// Aquire an ID but don't register yet		
-	WorldManager = World->GetSubsystem<UNetworkPredictionWorldManager>();
+	// Aquire an ID but don't register yet
+	WorldManager = Params.WorldManager;
 	npCheckSlow(WorldManager);
 
 	if (ID.IsValid() == false)
@@ -33,11 +33,12 @@ void FNetworkPredictionProxy::Init(UWorld* World, const FReplicationProxySet& Re
 		}
 
 		// Assign ID. Client will assign a temporary ID that later gets remapped via a call to ConfigFunc --> RemapClientSimulationID
-		ID = WorldManager->CreateSimulationID(World->GetNetMode() == NM_Client);
+		ID = WorldManager->CreateSimulationID(Params.Mode == NM_Client);
 	}
 
-	WorldManager->RegisterInstance<ModelDef>(ID, TNetworkPredictionModelInfo<ModelDef>(Simulation, Driver, &View));
+	WorldManager->RegisterInstance<ModelDef>(ID, TNetworkPredictionModelInfo<ModelDef>(Params.Simulation, Params.Driver, &View));
 
+	const FReplicationProxySet& RepProxies = Params.RepProxies;
 	ConfigFunc = [RepProxies](FNetworkPredictionProxy* const This, FNetworkPredictionID NewID, EConfigAction Action)
 	{
 		if (This->WorldManager == nullptr)
@@ -81,3 +82,10 @@ void FNetworkPredictionProxy::Init(UWorld* World, const FReplicationProxySet& Re
 		}
 	};
 }
+
+template<typename ModelDef>
+void FNetworkPredictionProxy::Init(UWorld* World, const FReplicationProxySet& RepProxies, typename ModelDef::Simulation* Simulation, typename ModelDef::Driver* Driver)
+{
+	FNetworkPredictionProxy::FInitParams<ModelDef> Params = {World->GetSubsystem<UNetworkPredictionWorldManager>(), World->GetNetMode(), RepProxies, Simulation, Driver};
+	Init<ModelDef>(Params);
+}
diff --git a/Engine/Plugins/Runtime/NetworkPrediction/Source/NetworkPrediction/Public/NetworkPredictionWorldManager.h b/Engine/Plugins/Runtime/NetworkPrediction/Source/NetworkPrediction/Public/NetworkPredictionWorldManager.h
--- a/Engine/Plugins/Runtime/NetworkPrediction/Source/NetworkPrediction/Public/NetworkPredictionWorldManager.h
+++ b/Engine/Plugins/Runtime/NetworkPrediction/Source/NetworkPrediction/Public/NetworkPredictionWorldManager.h
@@ -11,6 +11,11 @@
 class FChaosSolversModule;
 class ANetworkPredictionReplicatedManager;
 
+namespace UE::Net::Private
+{
+	struct FNetPredictionTestWorld;
+}
+
 UCLASS()
 class NETWORKPREDICTION_API UNetworkPredictionWorldManager : public UWorldSubsystem
 {
@@ -119,6 +124,12 @@ private:
 	void ReconcileSimulationsPostNetworkUpdate();
 	void BeginNewSimulationFrame(UWorld* InWorld, ELevelTick InLevelTick, float InDeltaSeconds);
 
+	// For ease of unit testing
+	friend struct UE::Net::Private::FNetPredictionTestWorld;
+	void OnWorldPreTick_Internal(float InDeltaSeconds, float InFixedFrameRate);
+	void ReconcileSimulationsPostNetworkUpdate_Internal();
+	void BeginNewSimulationFrame_Internal(float InDeltaSeconds);
+
 	FDelegateHandle PreTickDispatchHandle;
 	FDelegateHandle PostTickDispatchHandle;
 	FDelegateHandle PreWorldActorTickHandle;
diff --git a/Engine/Source/Programs/NetworkPredictionTests/Private/NetPredictionMockPackageMap.cpp b/Engine/Source/Programs/NetworkPredictionTests/Private/NetPredictionMockPackageMap.cpp
new file mode 100644
--- /dev/null
+++ b/Engine/Source/Programs/NetworkPredictionTests/Private/NetPredictionMockPackageMap.cpp
@@ -0,0 +1,61 @@
+// Copyright Epic Games, Inc. All Rights Reserved.
+
+#include "Misc/AssertionMacros.h"
+#include "NetPredictionMockPackageMap.h"
+
+UNetPredictionMockPackageMap* UNetPredictionMockPackageMap::Get()
+{
+	static UNetPredictionMockPackageMap* Instance = nullptr;
+	if (Instance == nullptr)
+	{
+		Instance = NewObject<UNetPredictionMockPackageMap>();
+	}
+	return Instance;
+}
+
+bool UNetPredictionMockPackageMap::WriteObject(FArchive & Ar, UObject* InOuter, FNetworkGUID NetGUID, FString ObjName)
+{
+	ensureAlwaysMsgf(false, TEXT("Not supported"));
+	return false;
+}
+
+bool UNetPredictionMockPackageMap::SerializeObject(FArchive& Ar, UClass* InClass, UObject*& Obj, FNetworkGUID *OutNetGUID)
+{
+	ensureAlwaysMsgf(false, TEXT("Not supported"));
+	return false;
+}
+
+bool UNetPredictionMockPackageMap::SerializeName(FArchive& Ar, FName& InName)
+{
+	ensureAlwaysMsgf(false, TEXT("Not supported"));
+	return false;
+}
+
+UObject* UNetPredictionMockPackageMap::ResolvePathAndAssignNetGUID(const FNetworkGUID& NetGUID, const FString& PathName)
+{
+	ensureAlwaysMsgf(false, TEXT("Not supported"));
+	return nullptr;
+}
+
+bool UNetPredictionMockPackageMap::SerializeNewActor(FArchive & Ar, class UActorChannel * Channel, class AActor *& Actor)
+{
+	ensureAlwaysMsgf(false, TEXT("Not supported"));
+	return false;
+}
+
+UObject* UNetPredictionMockPackageMap::GetObjectFromNetGUID(const FNetworkGUID& NetGUID, const bool bIgnoreMustBeMapped )
+{
+	ensureAlwaysMsgf(false, TEXT("Not supported"));
+	return nullptr;
+}
+
+FNetworkGUID UNetPredictionMockPackageMap::GetNetGUIDFromObject(const UObject* InObject) const
+{
+	ensureAlwaysMsgf(false, TEXT("Not supported"));
+	return FNetworkGUID();
+}
+
+void UNetPredictionMockPackageMap::Serialize(FArchive& Ar)
+{
+	ensureAlwaysMsgf(false, TEXT("Not supported"));
+}
diff --git a/Engine/Source/Programs/NetworkPredictionTests/Private/NetPredictionMockPackageMap.h b/Engine/Source/Programs/NetworkPredictionTests/Private/NetPredictionMockPackageMap.h
new file mode 100644
--- /dev/null
+++ b/Engine/Source/Programs/NetworkPredictionTests/Private/NetPredictionMockPackageMap.h
@@ -0,0 +1,25 @@
+// Copyright Epic Games, Inc. All Rights Reserved.
+
+#pragma once
+
+#include "UObject/CoreNet.h"
+#include "NetPredictionMockPackageMap.generated.h"
+
+// Placeholder package map that ensures if any of its functionality is actually used.
+UCLASS(transient)
+class UNetPredictionMockPackageMap : public UPackageMap
+{
+	GENERATED_BODY()
+
+public:
+	virtual bool WriteObject(FArchive & Ar, UObject* InOuter, FNetworkGUID NetGUID, FString ObjName) override;
+	virtual bool SerializeObject(FArchive& Ar, UClass* InClass, UObject*& Obj, FNetworkGUID *OutNetGUID = NULL) override;
+	virtual bool SerializeName(FArchive& Ar, FName& InName) override;
+	virtual UObject* ResolvePathAndAssignNetGUID(const FNetworkGUID& NetGUID, const FString& PathName) override;
+	virtual bool SerializeNewActor(FArchive & Ar, class UActorChannel * Channel, class AActor *& Actor) override;
+	virtual UObject* GetObjectFromNetGUID(const FNetworkGUID& NetGUID, const bool bIgnoreMustBeMapped) override;
+	virtual FNetworkGUID GetNetGUIDFromObject(const UObject* InObject) const override;
+	virtual void Serialize(FArchive& Ar) override;
+
+	static UNetPredictionMockPackageMap* Get();
+};
diff --git a/Engine/Source/Programs/NetworkPredictionTests/Private/NetPredictionTestChannel.cpp b/Engine/Source/Programs/NetworkPredictionTests/Private/NetPredictionTestChannel.cpp
new file mode 100644
--- /dev/null
+++ b/Engine/Source/Programs/NetworkPredictionTests/Private/NetPredictionTestChannel.cpp
@@ -0,0 +1,29 @@
+// Copyright Epic Games, Inc. All Rights Reserved.
+
+#include "NetPredictionTestChannel.h"
+#include "UObject/CoreNet.h"
+
+namespace UE::Net::Private
+{
+
+void FNetPredictionTestChannel::Send(const FNetBitWriter& Writer)
+{
+	DataBuffer.Enqueue(Writer);
+}
+
+bool FNetPredictionTestChannel::HasPendingData()
+{
+	return !DataBuffer.IsEmpty();
+}
+
+FNetBitReader FNetPredictionTestChannel::Receive()
+{
+	const FNetBitWriter& Writer = DataBuffer.Peek();
+	FNetBitReader Reader(nullptr, Writer.GetData(), Writer.GetNumBits());
+
+	DataBuffer.Pop();
+
+	return MoveTemp(Reader);
+}
+
+}
\ No newline at end of file
diff --git a/Engine/Source/Programs/NetworkPredictionTests/Private/NetPredictionTestChannel.h b/Engine/Source/Programs/NetworkPredictionTests/Private/NetPredictionTestChannel.h
new file mode 100644
--- /dev/null
+++ b/Engine/Source/Programs/NetworkPredictionTests/Private/NetPredictionTestChannel.h
@@ -0,0 +1,21 @@
+// Copyright Epic Games, Inc. All Rights Reserved.
+
+#pragma once
+
+#include "Net/Core/Misc/ResizableCircularQueue.h"
+#include "UObject/CoreNet.h"
+
+namespace UE::Net::Private
+{
+
+struct FNetPredictionTestChannel
+{
+	void Send(const FNetBitWriter& Writer);
+	bool HasPendingData();
+	FNetBitReader Receive();
+
+private:
+	TResizableCircularQueue<FNetBitWriter> DataBuffer;
+};
+
+}
\ No newline at end of file
diff --git a/Engine/Source/Programs/NetworkPredictionTests/Private/NetPredictionTestDriver.cpp b/Engine/Source/Programs/NetworkPredictionTests/Private/NetPredictionTestDriver.cpp
new file mode 100644
--- /dev/null
+++ b/Engine/Source/Programs/NetworkPredictionTests/Private/NetPredictionTestDriver.cpp
@@ -0,0 +1,163 @@
+// Copyright Epic Games, Inc. All Rights Reserved.
+
+#include "NetPredictionTestDriver.h"
+#include "NetPredictionMockPackageMap.h"
+#include "NetPredictionTestChannel.h"
+#include "NetworkPredictionModelDefRegistry.h"
+#include "NetworkPredictionProxyInit.h"
+#include "NetworkPredictionReplicationProxy.h"
+#include "NetworkPredictionTests.h"
+
+#include <catch2/catch_test_macros.hpp>
+
+namespace UE::Net::Private
+{
+
+struct FTestNetPredictionModelDef : public FNetworkPredictionModelDef
+{
+	NP_MODEL_BODY();
+
+	using Simulation = FNetPredictionTestDriver;
+	using StateTypes = FNetPredictionTestStateTypes;
+	using Driver = FNetPredictionTestDriver;
+
+	static const TCHAR* GetName() { return TEXT("TestModelDef"); }
+};
+
+NP_MODEL_REGISTER(FTestNetPredictionModelDef);
+
+bool FNetPredictionTestSyncState::ShouldReconcile(const FNetPredictionTestSyncState& AuthorityState) const
+{
+	return AutoCounter != AuthorityState.AutoCounter || InputCounter != AuthorityState.InputCounter;
+}
+
+void FNetPredictionTestSyncState::Interpolate(const FNetPredictionTestSyncState* From, const FNetPredictionTestSyncState* To, float Percent)
+{
+	UE_LOG(LogNetworkPredictionTests, Verbose, TEXT("Interpolate: from {%.2f, %.2f} to {%.2f, %.2f}, pct %.2f."),
+		From->AutoCounter, From->InputCounter, To->AutoCounter, To->InputCounter, Percent);
+
+	AutoCounter = FMath::Lerp(From->AutoCounter, To->AutoCounter, Percent);
+	InputCounter = FMath::Lerp(From->InputCounter, To->InputCounter, Percent);
+}
+
+FNetPredictionTestDriver::FNetPredictionTestDriver(UNetworkPredictionWorldManager* WorldManager, ENetMode Mode,
+	TSharedPtr<FNetPredictionTestChannel> InClientToServer,
+	TSharedPtr<FNetPredictionTestChannel> InServerToClient)
+	: ClientToServer(InClientToServer)
+	, ServerToClient(InServerToClient)
+{
+	ReplicationProxy_ServerRPC.Init(&Proxy, EReplicationProxyTarget::ServerRPC);
+	ReplicationProxy_Autonomous.Init(&Proxy, EReplicationProxyTarget::AutonomousProxy);
+	ReplicationProxy_Simulated.Init(&Proxy, EReplicationProxyTarget::SimulatedProxy);
+
+	FNetworkPredictionProxy::FInitParams<FTestNetPredictionModelDef> Params = {WorldManager, Mode, GetReplicationProxies(), this, this};
+	Proxy.Init<FTestNetPredictionModelDef>(Params);
+}
+
+void FNetPredictionTestDriver::ProduceInput(const int32 DeltaTimeMS, FNetPredictionTestInputCmd* Cmd)
+{
+	UE_LOG(LogNetworkPredictionTests, Verbose, TEXT("ProduceInput: %s. SimFrame: %d"), *DebugName, Proxy.GetPendingFrame());
+
+	Cmd->bIncrement = bInputPressed;
+}
+
+void FNetPredictionTestDriver::SimulationTick(const FNetSimTimeStep& TimeStep, const TNetSimInput<FNetPredictionTestStateTypes>& SimInput, const TNetSimOutput<FNetPredictionTestStateTypes>& SimOutput)
+{
+	UE_LOG(LogNetworkPredictionTests, Verbose, TEXT("SimulationTick: %s. Frame: %d"), *DebugName, TimeStep.Frame);
+
+	SimOutput.Sync->InputCounter = SimInput.Sync->InputCounter;
+	if (SimInput.Cmd->bIncrement)
+	{
+		SimOutput.Sync->InputCounter += 1.0f;
+	}
+
+	SimOutput.Sync->AutoCounter = SimInput.Sync->AutoCounter + 1.0f;
+}
+
+void FNetPredictionTestDriver::InitializeSimulationState(FNetPredictionTestSyncState* OutSync, void* OutAux)
+{
+	UE_LOG(LogNetworkPredictionTests, Verbose, TEXT("InitializeSimulationState: %s"), *DebugName);
+	OutSync->AutoCounter = 0.0f;
+	OutSync->InputCounter = 0.0f;
+}
+
+void FNetPredictionTestDriver::FinalizeFrame(const FNetPredictionTestSyncState* SyncState, const void* AuxState)
+{
+	UE_LOG(LogNetworkPredictionTests, Verbose, TEXT("FinalizeFrame: %s. AutoCounter: %.2f, InputCounter: %.2f"), *DebugName, SyncState->AutoCounter, SyncState->InputCounter);
+}
+
+void FNetPredictionTestDriver::SetHiddenForInterpolation(bool bInHidden)
+{
+	bHidden = bInHidden;
+}
+
+void FNetPredictionTestDriver::CallServerRPC()
+{
+	FServerReplicationRPCParameter ProxyParameter(ReplicationProxy_ServerRPC);
+
+	FNetBitWriter Ar(1024 * 8 * 2);
+	bool bOutSuccess = false;
+	ProxyParameter.NetSerialize(Ar, UNetPredictionMockPackageMap::Get(), bOutSuccess);
+
+	ClientToServer->Send(Ar);
+}
+
+void FNetPredictionTestDriver::ReceiveServerRPCs()
+{
+	while (ClientToServer->HasPendingData())
+	{
+		FNetBitReader Ar = ClientToServer->Receive();
+
+		FServerReplicationRPCParameter ProxyParameter;
+		bool bOutSuccess = false;
+		ProxyParameter.NetSerialize(Ar, UNetPredictionMockPackageMap::Get(), bOutSuccess);
+		ProxyParameter.NetSerializeToProxy(ReplicationProxy_ServerRPC);
+	}
+}
+
+void FNetPredictionTestDriver::TestNetSerialize(FArchive& Ar)
+{
+	bool bSuccess = false;
+
+	Proxy.NetSerialize(Ar, nullptr, bSuccess);
+
+	// Emulate the replication conditions from UNetworkPredictionComponent
+	if (Proxy.GetCachedHasNetConnection())
+	{
+		ReplicationProxy_Autonomous.NetSerialize(Ar, UNetPredictionMockPackageMap::Get(), bSuccess);
+	}
+	else
+	{
+		ReplicationProxy_Simulated.NetSerialize(Ar, UNetPredictionMockPackageMap::Get(), bSuccess);
+	}
+}
+
+void FNetPredictionTestDriver::ServerSend()
+{
+	FNetBitWriter TempWriter(1024 * 8 * 2);
+	TestNetSerialize(TempWriter);
+	ServerToClient->Send(TempWriter);
+}
+
+void FNetPredictionTestDriver::ClientReceive()
+{
+	REQUIRE(ServerToClient->HasPendingData());
+
+	FNetBitReader TempReader = ServerToClient->Receive();
+	TestNetSerialize(TempReader);
+}
+
+FNetPredictionTestObject::FNetPredictionTestObject(
+		UNetworkPredictionWorldManager* ServerWorldManager,
+		UNetworkPredictionWorldManager* ClientWorldManager,
+		ENetRole ClientRole)
+	: ClientToServer(MakeShared<FNetPredictionTestChannel>())
+	, ServerToClient(MakeShared<FNetPredictionTestChannel>())
+	, ServerObject(ServerWorldManager, NM_DedicatedServer, ClientToServer, ServerToClient)
+	, ClientObject(ClientWorldManager, NM_Client, ClientToServer, ServerToClient)
+{
+	ServerObject.Proxy.InitForNetworkRole(ENetRole::ROLE_Authority, ClientRole == ENetRole::ROLE_AutonomousProxy);
+	ClientObject.Proxy.InitForNetworkRole(ClientRole, ClientRole == ENetRole::ROLE_AutonomousProxy);
+}
+
+}
\ No newline at end of file
diff --git a/Engine/Source/Programs/NetworkPredictionTests/Private/NetPredictionTestDriver.h b/Engine/Source/Programs/NetworkPredictionTests/Private/NetPredictionTestDriver.h
new file mode 100644
--- /dev/null
+++ b/Engine/Source/Programs/NetworkPredictionTests/Private/NetPredictionTestDriver.h
@@ -0,0 +1,114 @@
+// Copyright Epic Games, Inc. All Rights Reserved.
+
+#pragma once
+
+#include "Engine/EngineBaseTypes.h"
+#include "NetPredictionTestChannel.h"
+#include "NetworkPredictionModelDef.h"
+#include "NetworkPredictionProxy.h"
+#include "NetworkPredictionReplicationProxy.h"
+#include "NetworkPredictionSimulation.h"
+#include "NetworkPredictionTickState.h"
+
+class FArchive;
+class UNetworkPredictionWorldManager;
+
+namespace UE::Net::Private
+{
+
+struct FNetPredictionTestInputCmd
+{
+	bool bIncrement = false;
+
+	void ToString(FAnsiStringBuilderBase& Builder) const
+	{
+		Builder << TEXT("bIncrement: ") << bIncrement;
+	}
+
+	void NetSerialize(const FNetSerializeParams& P)
+	{
+		P.Ar.SerializeBits(&bIncrement, 1);
+	}
+};
+
+struct FNetPredictionTestSyncState
+{
+	// Value is incremented by one every frame, not dependent on input
+	float AutoCounter = 0.0f;
+
+	// Value is incremented by one if the input bIncrement value is true, no change if it's false.
+	float InputCounter = 0.0f;
+
+	void ToString(FAnsiStringBuilderBase& Builder) const
+	{
+		Builder << FString::Printf(TEXT("AutoCounter: %.2f, InputCounter: %.2f"), AutoCounter, InputCounter);
+	}
+
+	void NetSerialize(const FNetSerializeParams& P)
+	{
+		P.Ar << AutoCounter;
+		P.Ar << InputCounter;
+	}
+
+	void Interpolate(const FNetPredictionTestSyncState* From, const FNetPredictionTestSyncState* To, float Percent);
+	bool ShouldReconcile(const FNetPredictionTestSyncState& AuthorityState) const;
+};
+
+using FNetPredictionTestStateTypes = TNetworkPredictionStateTypes<FNetPredictionTestInputCmd, FNetPredictionTestSyncState, void>;
+
+// Mimics the key parts of UNetworkPredictionComponent without being an actual UActorComponent.
+struct FNetPredictionTestDriver
+{
+	FNetworkPredictionProxy Proxy;
+	FReplicationProxy ReplicationProxy_ServerRPC;
+	FReplicationProxy ReplicationProxy_Autonomous;
+	FReplicationProxy ReplicationProxy_Simulated;
+
+	FString DebugName;
+
+	bool bHidden = false;
+	bool bInputPressed = false;
+
+	FNetPredictionTestDriver(UNetworkPredictionWorldManager* WorldManager, ENetMode Mode,
+		TSharedPtr<FNetPredictionTestChannel> InClientToServer,
+		TSharedPtr<FNetPredictionTestChannel> InServerToClient);
+
+	void ProduceInput(const int32 DeltaTimeMS, FNetPredictionTestInputCmd* Cmd);
+	void SimulationTick(const FNetSimTimeStep& TimeStep, const TNetSimInput<FNetPredictionTestStateTypes>& SimInput, const TNetSimOutput<FNetPredictionTestStateTypes>& SimOutput);
+	void InitializeSimulationState(FNetPredictionTestSyncState* OutSync, void* OutAux);
+	void FinalizeFrame(const FNetPredictionTestSyncState* SyncState, const void* AuxState);
+	void SetHiddenForInterpolation(bool bInHidden);
+	void CallServerRPC();
+	void ReceiveServerRPCs();
+
+	void TestNetSerialize(FArchive& Ar);
+
+	void ServerSend();
+	void ClientReceive();
+
+private:
+	FReplicationProxySet GetReplicationProxies()
+	{
+		return FReplicationProxySet{&ReplicationProxy_ServerRPC, &ReplicationProxy_Autonomous, &ReplicationProxy_Simulated, nullptr};
+	}
+
+	TSharedPtr<FNetPredictionTestChannel> ClientToServer;
+	TSharedPtr<FNetPredictionTestChannel> ServerToClient;
+};
+
+// Conveniently associates the server & client object drivers and their channels for ease of writing unit tests.
+struct FNetPredictionTestObject
+{
+	TSharedPtr<FNetPredictionTestChannel> ClientToServer;
+	TSharedPtr<FNetPredictionTestChannel> ServerToClient;
+
+	FNetPredictionTestDriver ServerObject;
+	FNetPredictionTestDriver ClientObject;
+
+	FNetPredictionTestObject(
+		UNetworkPredictionWorldManager* ServerWorldManager,
+		UNetworkPredictionWorldManager* ClientWorldManager,
+		ENetRole ClientRole);
+};
+
+}
\ No newline at end of file
diff --git a/Engine/Source/Programs/NetworkPredictionTests/Private/NetPredictionTestWorld.cpp b/Engine/Source/Programs/NetworkPredictionTests/Private/NetPredictionTestWorld.cpp
new file mode 100644
--- /dev/null
+++ b/Engine/Source/Programs/NetworkPredictionTests/Private/NetPredictionTestWorld.cpp
@@ -0,0 +1,35 @@
+// Copyright Epic Games, Inc. All Rights Reserved.
+
+#include "NetPredictionTestWorld.h"
+#include "NetworkPredictionWorldManager.h"
+
+namespace UE::Net::Private
+{
+
+FNetPredictionTestWorld::FNetPredictionTestWorld(UNetworkPredictionSettingsObject* Settings, ENetMode InNetMode)
+	: NetMode(InNetMode)
+{
+	WorldManager = NewObject<UNetworkPredictionWorldManager>();
+	if (Settings)
+	{
+		WorldManager->SyncNetworkPredictionSettings(Settings);
+	}
+}
+
+void FNetPredictionTestWorld::PreTick(int32 FixedFrameRate)
+{
+	checkf(FixedFrameRate > 0, TEXT("FixedFrameRate must be > 0"));
+	WorldManager->OnWorldPreTick_Internal(1.0f / FixedFrameRate, FixedFrameRate);
+}
+
+void FNetPredictionTestWorld::Tick(int32 FixedFrameRate)
+{
+	checkf(FixedFrameRate > 0, TEXT("FixedFrameRate must be > 0"));
+	if (NetMode == NM_Client)
+	{
+		WorldManager->ReconcileSimulationsPostNetworkUpdate_Internal();
+	}
+	WorldManager->BeginNewSimulationFrame_Internal(1.0f / FixedFrameRate);
+}
+
+}
\ No newline at end of file
diff --git a/Engine/Source/Programs/NetworkPredictionTests/Private/NetPredictionTestWorld.h b/Engine/Source/Programs/NetworkPredictionTests/Private/NetPredictionTestWorld.h
new file mode 100644
--- /dev/null
+++ b/Engine/Source/Programs/NetworkPredictionTests/Private/NetPredictionTestWorld.h
@@ -0,0 +1,27 @@
+// Copyright Epic Games, Inc. All Rights Reserved.
+
+#pragma once
+
+#include "Engine/EngineBaseTypes.h"
+
+class UNetworkPredictionWorldManager;
+class UNetworkPredictionSettingsObject;
+
+namespace UE::Net::Private
+{
+
+struct FNetPredictionTestWorld
+{
+public:
+	FNetPredictionTestWorld(UNetworkPredictionSettingsObject* Settings, ENetMode InNetMode);
+
+	void PreTick(int32 FixedFrameRate);
+	void Tick(int32 FixedFrameRate);
+
+	UNetworkPredictionWorldManager* WorldManager = nullptr;
+
+private:
+	ENetMode NetMode;
+};
+
+}
\ No newline at end of file
diff --git a/Engine/Source/Programs/NetworkPredictionTests/Private/NetworkPredictionTests.cpp b/Engine/Source/Programs/NetworkPredictionTests/Private/NetworkPredictionTests.cpp
--- a/Engine/Source/Programs/NetworkPredictionTests/Private/NetworkPredictionTests.cpp
+++ b/Engine/Source/Programs/NetworkPredictionTests/Private/NetworkPredictionTests.cpp
@@ -1,14 +1,16 @@
 // Copyright Epic Games, Inc. All Rights Reserved.
 
-
 #include "Logging/LogScopedVerbosityOverride.h"
 #include "Modules/ModuleManager.h"
+#include "NetworkPredictionTests.h"
 #include "SlateGlobals.h"
 #include "Styling/SlateWidgetStyleContainerBase.h"
 #include "TestCommon/Initialization.h"
 
 #include <catch2/catch_test_macros.hpp>
 
+DEFINE_LOG_CATEGORY(LogNetworkPredictionTests);
+
 GROUP_BEFORE_GLOBAL(Catch::DefaultGroup)
 {
 	{
diff --git a/Engine/Source/Programs/NetworkPredictionTests/Private/NetworkPredictionTests.h b/Engine/Source/Programs/NetworkPredictionTests/Private/NetworkPredictionTests.h
new file mode 100644
--- /dev/null
+++ b/Engine/Source/Programs/NetworkPredictionTests/Private/NetworkPredictionTests.h
@@ -0,0 +1,7 @@
+// Copyright Epic Games, Inc. All Rights Reserved.
+
+#pragma once
+
+#include "Logging/LogMacros.h"
+
+DECLARE_LOG_CATEGORY_EXTERN(LogNetworkPredictionTests, Log, All);
diff --git a/Engine/Source/Programs/NetworkPredictionTests/Private/WorldManagerTests_ForwardPredict.cpp b/Engine/Source/Programs/NetworkPredictionTests/Private/WorldManagerTests_ForwardPredict.cpp
new file mode 100644
--- /dev/null
+++ b/Engine/Source/Programs/NetworkPredictionTests/Private/WorldManagerTests_ForwardPredict.cpp
@@ -0,0 +1,224 @@
+// Copyright Epic Games, Inc. All Rights Reserved.
+
+#include "NetPredictionTestDriver.h"
+#include "NetPredictionTestWorld.h"
+#include "NetworkPredictionDriver.h"
+#include "NetworkPredictionProxy.h"
+
+#include <catch2/catch_test_macros.hpp>
+
+namespace UE::Net::Private
+{
+
+constexpr int32 PredictionTestFixedFrameRate = 60;
+
+TEST_CASE("NetworkPrediction, single authority object")
+{
+	FNetPredictionTestWorld TestWorld(nullptr, NM_DedicatedServer);
+
+	FNetPredictionTestDriver SimObject(TestWorld.WorldManager, NM_DedicatedServer, nullptr, nullptr);
+	SimObject.Proxy.InitForNetworkRole(ENetRole::ROLE_Authority, false);
+
+	float ExpectedAutoCounterValue = 0.0f;
+
+	const FNetPredictionTestSyncState* TestState = SimObject.Proxy.ReadSyncState<FNetPredictionTestSyncState>();
+	REQUIRE(TestState->AutoCounter == ExpectedAutoCounterValue);
+
+	SECTION("Simulation advances on ticks")
+	{
+		for (int i = 0; i < 10; ++i)
+		{
+			TestWorld.PreTick(PredictionTestFixedFrameRate);
+			TestWorld.Tick(PredictionTestFixedFrameRate);
+			ExpectedAutoCounterValue += 1.0f;
+
+			TestState = SimObject.Proxy.ReadSyncState<FNetPredictionTestSyncState>();
+			REQUIRE(TestState->AutoCounter == ExpectedAutoCounterValue);
+		}
+	}
+
+	SECTION("Simulation processes input")
+	{
+		float ExpectedInputCounterValue = 0.0f;
+
+		REQUIRE(TestState->InputCounter == ExpectedInputCounterValue);
+
+		// Some frames with input actions
+		SimObject.bInputPressed = true;
+		for (int i = 0; i < 10; ++i)
+		{
+			TestWorld.PreTick(PredictionTestFixedFrameRate);
+			TestWorld.Tick(PredictionTestFixedFrameRate);
+			ExpectedAutoCounterValue += 1.0f;
+			ExpectedInputCounterValue += 1.0f;
+
+			TestState = SimObject.Proxy.ReadSyncState<FNetPredictionTestSyncState>();
+			REQUIRE(TestState->AutoCounter == ExpectedAutoCounterValue);
+			REQUIRE(TestState->InputCounter == ExpectedInputCounterValue);
+		}
+
+		// Some frames without input actions
+		SimObject.bInputPressed = false;
+		for (int i = 0; i < 10; ++i)
+		{
+			TestWorld.PreTick(PredictionTestFixedFrameRate);
+			TestWorld.Tick(PredictionTestFixedFrameRate);
+			ExpectedAutoCounterValue += 1.0f;
+
+			TestState = SimObject.Proxy.ReadSyncState<FNetPredictionTestSyncState>();
+			REQUIRE(TestState->AutoCounter == ExpectedAutoCounterValue);
+			REQUIRE(TestState->InputCounter == ExpectedInputCounterValue);
+		}
+
+		// And some more with input actions
+		SimObject.bInputPressed = true;
+		for (int i = 0; i < 10; ++i)
+		{
+			TestWorld.PreTick(PredictionTestFixedFrameRate);
+			TestWorld.Tick(PredictionTestFixedFrameRate);
+			ExpectedAutoCounterValue += 1.0f;
+			ExpectedInputCounterValue += 1.0f;
+
+			TestState = SimObject.Proxy.ReadSyncState<FNetPredictionTestSyncState>();
+			REQUIRE(TestState->AutoCounter == ExpectedAutoCounterValue);
+			REQUIRE(TestState->InputCounter == ExpectedInputCounterValue);
+		}
+	}
+}
+
+struct FSingleSimulatedObjectWorld
+{
+	FNetPredictionTestWorld ServerWorld;
+	FNetPredictionTestWorld ClientWorld;
+
+	FNetPredictionTestObject SimObject;
+
+	FSingleSimulatedObjectWorld(ENetRole ClientRole)
+		: ServerWorld(nullptr, NM_DedicatedServer)
+		, ClientWorld(nullptr, NM_Client)
+		, SimObject(ServerWorld.WorldManager, ClientWorld.WorldManager, ClientRole)
+	{
+	}
+
+	void TickServer()
+	{
+		ServerWorld.PreTick(PredictionTestFixedFrameRate);
+		SimObject.ServerObject.ReceiveServerRPCs();
+		ServerWorld.Tick(PredictionTestFixedFrameRate);
+		SimObject.ServerObject.ServerSend();
+	}
+
+	void TickClient()
+	{
+		ClientWorld.PreTick(PredictionTestFixedFrameRate);
+		SimObject.ClientObject.ClientReceive();
+		ClientWorld.Tick(PredictionTestFixedFrameRate);
+	}
+};
+
+// For this test currently the client will roll back every frame. Looks like the system currently relies
+// on having an automonous proxy to sync frame numbers. Disabled until that's addressed.
+TEST_CASE("NetworkPrediction, one sim proxy object updates sim state on client", "[.][warnings]")
+{
+	constexpr ENetRole ClientRole = ENetRole::ROLE_SimulatedProxy;
+
+	FSingleSimulatedObjectWorld TestWorld(ENetRole::ROLE_SimulatedProxy);
+
+	float ExpectedAutoCounterValue = 0.0f;
+
+	const FNetPredictionTestSyncState* TestState = TestWorld.SimObject.ServerObject.Proxy.ReadSyncState<FNetPredictionTestSyncState>();
+	REQUIRE(TestState->AutoCounter == ExpectedAutoCounterValue);
+
+	SECTION("Simulation advances on ticks")
+	{
+		for (int i = 0; i < 50; ++i)
+		{
+			TestWorld.TickServer();
+			ExpectedAutoCounterValue += 1.0f;
+
+			TestState = TestWorld.SimObject.ServerObject.Proxy.ReadSyncState<FNetPredictionTestSyncState>();
+			REQUIRE(TestState->AutoCounter == ExpectedAutoCounterValue);
+
+			TestWorld.TickClient();
+
+			const FNetPredictionTestSyncState* ClientState = TestWorld.SimObject.ClientObject.Proxy.ReadSyncState<FNetPredictionTestSyncState>();
+			REQUIRE(ClientState->AutoCounter == ExpectedAutoCounterValue);
+		}
+	}
+}
+
+TEST_CASE("NetworkPrediction, one autonomous proxy object updates sim state on client & server")
+{
+	constexpr ENetRole ClientRole = ENetRole::ROLE_AutonomousProxy;
+	FSingleSimulatedObjectWorld TestWorld(ENetRole::ROLE_AutonomousProxy);
+
+	float ExpectedServerAutoCounterValue = 0.0f;
+
+	const FNetPredictionTestSyncState* ServerState = TestWorld.SimObject.ServerObject.Proxy.ReadSyncState<FNetPredictionTestSyncState>();
+	REQUIRE(ServerState->AutoCounter == ExpectedServerAutoCounterValue);
+
+	// Run some frames to let things settle, will get the client ahead
+	// of the server by 6 frames
+	for (int i = 0; i < 10; ++i)
+	{
+		TestWorld.TickServer();
+		ExpectedServerAutoCounterValue += 1.0f;
+
+		ServerState = TestWorld.SimObject.ServerObject.Proxy.ReadSyncState<FNetPredictionTestSyncState>();
+		CHECK(ServerState->AutoCounter == ExpectedServerAutoCounterValue);
+		CHECK(ServerState->InputCounter == 0.0f);
+
+		TestWorld.TickClient();
+	}
+
+	// Client is now ahead of server by 6 frames.
+	SECTION("Server & client simulation, no client input")
+	{
+		for (int i = 0; i < 20; ++i)
+		{
+			TestWorld.TickServer();
+			ExpectedServerAutoCounterValue += 1.0f;
+
+			ServerState = TestWorld.SimObject.ServerObject.Proxy.ReadSyncState<FNetPredictionTestSyncState>();
+			CHECK(ServerState->AutoCounter == ExpectedServerAutoCounterValue);
+			CHECK(ServerState->InputCounter == 0.0f);
+
+			TestWorld.TickClient();
+
+			const FNetPredictionTestSyncState* ClientState = TestWorld.SimObject.ClientObject.Proxy.ReadSyncState<FNetPredictionTestSyncState>();
+			CHECK(ClientState->AutoCounter == ExpectedServerAutoCounterValue + 6.0f);
+			CHECK(ClientState->InputCounter == 0.0f);
+		}
+	}
+
+	SECTION("Server & client simulation, with client input")
+	{
+		TestWorld.SimObject.ClientObject.bInputPressed = true;
+		float ExpectedClientInputCounterValue = 0.0f;
+		float ExpectedServerInputCounterValue = 0.0f;
+
+		for (int i = 0; i < 20; ++i)
+		{
+			TestWorld.TickServer();
+			ExpectedServerAutoCounterValue += 1.0f;
+
+			if (i >= 6)
+			{
+				ExpectedServerInputCounterValue += 1.0f;
+			}
+
+			ServerState = TestWorld.SimObject.ServerObject.Proxy.ReadSyncState<FNetPredictionTestSyncState>();
+			CHECK(ServerState->AutoCounter == ExpectedServerAutoCounterValue);
+			CHECK(ServerState->InputCounter == ExpectedServerInputCounterValue);
+
+			TestWorld.TickClient();
+			ExpectedClientInputCounterValue += 1.0f;
+
+			const FNetPredictionTestSyncState* ClientState = TestWorld.SimObject.ClientObject.Proxy.ReadSyncState<FNetPredictionTestSyncState>();
+			CHECK(ClientState->AutoCounter == ExpectedServerAutoCounterValue + 6.0f);
+			CHECK(ClientState->InputCounter == ExpectedClientInputCounterValue);
+		}
+	}
+}
+
+}
\ No newline at end of file
diff --git a/Engine/Source/Programs/NetworkPredictionTests/Private/WorldManagerTests_Interpolation.cpp b/Engine/Source/Programs/NetworkPredictionTests/Private/WorldManagerTests_Interpolation.cpp
new file mode 100644
--- /dev/null
+++ b/Engine/Source/Programs/NetworkPredictionTests/Private/WorldManagerTests_Interpolation.cpp
@@ -0,0 +1,227 @@
+// Copyright Epic Games, Inc. All Rights Reserved.
+
+#include "NetPredictionTestDriver.h"
+#include "NetPredictionTestWorld.h"
+#include "NetworkPredictionDriver.h"
+#include "NetworkPredictionProxy.h"
+
+#include <catch2/catch_test_macros.hpp>
+
+namespace UE::Net::Private
+{
+
+struct FTestPossession
+{
+	FNetPredictionTestWorld ServerWorld;
+	FNetPredictionTestWorld ClientWorld;
+
+	FNetPredictionTestObject Object1;
+	FNetPredictionTestObject Object2;
+
+	FTestPossession(UNetworkPredictionSettingsObject* Settings)
+		: ServerWorld(Settings, NM_DedicatedServer)
+		, ClientWorld(Settings, NM_Client)
+		, Object1(ServerWorld.WorldManager, ClientWorld.WorldManager, ROLE_AutonomousProxy)
+		, Object2(ServerWorld.WorldManager, ClientWorld.WorldManager, ROLE_SimulatedProxy)
+	{
+	}
+
+	void TickServer()
+	{
+		ServerWorld.PreTick(FixedFrameRate);
+		Object1.ServerObject.ReceiveServerRPCs();
+		Object2.ServerObject.ReceiveServerRPCs();
+		ServerWorld.Tick(FixedFrameRate);
+		Object1.ServerObject.ServerSend();
+		Object2.ServerObject.ServerSend();
+	}
+
+	void TickClient()
+	{
+		ClientWorld.PreTick(FixedFrameRate);
+		Object1.ClientObject.ClientReceive();
+		Object2.ClientObject.ClientReceive();
+		ClientWorld.Tick(FixedFrameRate);
+	}
+
+private:
+	static constexpr int32 FixedFrameRate = 60;
+};
+
+// This case includes a repro for UE-170936, disabled until it's fixed.
+TEST_CASE("NetworkPrediction interpolation mode, one simulated proxy", "[.][possession]")
+{
+	UNetworkPredictionSettingsObject* Settings = NewObject<UNetworkPredictionSettingsObject>();
+	Settings->Settings.SimulatedProxyNetworkLOD = ENetworkLOD::Interpolated;
+
+	FTestPossession Worlds(Settings);
+
+	float ExpectedServerAutoCounterValue = 0.0f;
+
+	const FNetPredictionTestSyncState* ServerState = nullptr;
+
+	// Interpolated objects will buffer FNetworkPredictionSettings::FixedTickInterpolationBufferedMS
+	// frames worth of updates. For this test with the default of 100ms and fixed 60hz tick that's 6 frames.
+	constexpr int32 ExpectedInterpolateBufferFrames = 6;
+
+	SECTION("One simulated proxy")
+	{
+		// Run some frames to allow the client to buffer interpolation data
+		for (int i = 0; i < 20; ++i)
+		{
+			Worlds.TickServer();
+			ExpectedServerAutoCounterValue += 1.0f;
+
+			ServerState = Worlds.Object2.ServerObject.Proxy.ReadSyncState<FNetPredictionTestSyncState>();
+			CHECK(ServerState->AutoCounter == ExpectedServerAutoCounterValue);
+			CHECK(ServerState->InputCounter == 0.0f);
+
+			Worlds.TickClient();
+		}
+
+		// Client has now buffered and is behind the server by 6 frames.
+		const FNetPredictionTestSyncState* ClientState = Worlds.Object2.ClientObject.Proxy.ReadSyncState<FNetPredictionTestSyncState>();
+		CHECK(ClientState->AutoCounter == ExpectedServerAutoCounterValue - ExpectedInterpolateBufferFrames);
+		CHECK(ClientState->InputCounter == 0.0f);
+
+		// Run some more frames
+		for (int i = 0; i < 20; ++i)
+		{
+			Worlds.TickServer();
+			ExpectedServerAutoCounterValue += 1.0f;
+
+			ServerState = Worlds.Object2.ServerObject.Proxy.ReadSyncState<FNetPredictionTestSyncState>();
+			CHECK(ServerState->AutoCounter == ExpectedServerAutoCounterValue);
+			CHECK(ServerState->InputCounter == 0.0f);
+
+			Worlds.TickClient();
+			ClientState = Worlds.Object2.ClientObject.Proxy.ReadSyncState<FNetPredictionTestSyncState>();
+			CHECK(ClientState->AutoCounter == ExpectedServerAutoCounterValue - ExpectedInterpolateBufferFrames);
+			CHECK(ClientState->InputCounter == 0.0f);
+		}
+	}
+	
+	SECTION("Changing possession")
+	{
+		// This section repros UE-170936
+		const FNetPredictionTestSyncState* ServerState1 = nullptr;
+		const FNetPredictionTestSyncState* ServerState2 = nullptr;
+
+		Worlds.Object1.ServerObject.DebugName = TEXT("Obj1Server");
+		Worlds.Object1.ClientObject.DebugName = TEXT("Obj1Client");
+		Worlds.Object2.ServerObject.DebugName = TEXT("Obj2Server");
+		Worlds.Object2.ClientObject.DebugName = TEXT("Obj2Client");
+
+		constexpr int32 ExpectedForwardPredictFrames = 6;
+
+		// Run some frames to allow the client to predict ahead (autonomous proxy)
+		// and buffer interpolation data (simulated proxy)
+		for (int i = 0; i < 20; ++i)
+		{
+			Worlds.TickServer();
+			ExpectedServerAutoCounterValue += 1.0f;
+
+			ServerState1 = Worlds.Object1.ServerObject.Proxy.ReadSyncState<FNetPredictionTestSyncState>();
+			CHECK(ServerState1->AutoCounter == ExpectedServerAutoCounterValue);
+			CHECK(ServerState1->InputCounter == 0.0f);
+
+			ServerState2 = Worlds.Object2.ServerObject.Proxy.ReadSyncState<FNetPredictionTestSyncState>();
+			CHECK(ServerState2->AutoCounter == ExpectedServerAutoCounterValue);
+			CHECK(ServerState2->InputCounter == 0.0f);
+
+			Worlds.TickClient();
+		}
+
+		const FNetPredictionTestSyncState* ClientState1 = Worlds.Object1.ClientObject.Proxy.ReadSyncState<FNetPredictionTestSyncState>();
+		const FNetPredictionTestSyncState* ClientState2 = Worlds.Object2.ClientObject.Proxy.ReadSyncState<FNetPredictionTestSyncState>();
+		// Auto proxy is ahead of the server by 6 frames.
+		CHECK(ClientState1->AutoCounter == ExpectedServerAutoCounterValue + ExpectedForwardPredictFrames);
+		CHECK(ClientState1->InputCounter == 0.0f);
+		// Simulated proxy is behind the server by 6 frames.
+		CHECK(ClientState2->AutoCounter == ExpectedServerAutoCounterValue - ExpectedInterpolateBufferFrames);
+		CHECK(ClientState2->InputCounter == 0.0f);
+
+		// "Unpossess" the current autonomous proxy on server & client
+		Worlds.Object1.ServerObject.Proxy.InitForNetworkRole(ROLE_Authority, false);
+		Worlds.Object1.ClientObject.Proxy.InitForNetworkRole(ROLE_SimulatedProxy, false);
+
+		float ExpectedClientAutoCounterValue = 21.0f;
+
+		// Continue running
+		for (int i = 0; i < 80; ++i)
+		{
+			Worlds.TickServer();
+			ExpectedServerAutoCounterValue += 1.0f;
+
+			ServerState1 = Worlds.Object1.ServerObject.Proxy.ReadSyncState<FNetPredictionTestSyncState>();
+			CHECK(ServerState1->AutoCounter == ExpectedServerAutoCounterValue);
+			CHECK(ServerState1->InputCounter == 0.0f);
+
+			ServerState2 = Worlds.Object2.ServerObject.Proxy.ReadSyncState<FNetPredictionTestSyncState>();
+			CHECK(ServerState2->AutoCounter == ExpectedServerAutoCounterValue);
+			CHECK(ServerState2->InputCounter == 0.0f);
+
+			Worlds.TickClient();
+
+			// Object1 went from autonomous proxy to simulated proxy
+			ClientState1 = Worlds.Object1.ClientObject.Proxy.ReadSyncState<FNetPredictionTestSyncState>();
+			// For the first few frames it's buffering interpolation data
+			if (i < ExpectedInterpolateBufferFrames)
+			{
+				CHECK(ExpectedClientAutoCounterValue == 21.0f);
+			}
+			else
+			{
+				CHECK(ClientState1->AutoCounter == ExpectedServerAutoCounterValue - ExpectedInterpolateBufferFrames);
+			}
+			CHECK(ClientState1->InputCounter == 0.0f);
+			
+			// Object2 is still a simulated proxy
+			ClientState2 = Worlds.Object2.ClientObject.Proxy.ReadSyncState<FNetPredictionTestSyncState>();
+			CHECK(ClientState2->AutoCounter == ExpectedServerAutoCounterValue - ExpectedInterpolateBufferFrames);
+			CHECK(ClientState2->InputCounter == 0.0f);
+		}
+
+		// "Possess" the other object
+		Worlds.Object2.ServerObject.Proxy.InitForNetworkRole(ROLE_Authority, true);
+		Worlds.Object2.ClientObject.Proxy.InitForNetworkRole(ROLE_AutonomousProxy, true);
+
+		// Continue running
+		for (int i = 0; i < 80; ++i)
+		{
+			Worlds.TickServer();
+			ExpectedServerAutoCounterValue += 1.0f;
+
+			ServerState1 = Worlds.Object1.ServerObject.Proxy.ReadSyncState<FNetPredictionTestSyncState>();
+			CHECK(ServerState1->AutoCounter == ExpectedServerAutoCounterValue);
+			CHECK(ServerState1->InputCounter == 0.0f);
+
+			ServerState2 = Worlds.Object2.ServerObject.Proxy.ReadSyncState<FNetPredictionTestSyncState>();
+			CHECK(ServerState2->AutoCounter == ExpectedServerAutoCounterValue);
+			CHECK(ServerState2->InputCounter == 0.0f);
+
+			Worlds.TickClient();
+
+			// Object1 is still a simulated proxy
+			ClientState1 = Worlds.Object1.ClientObject.Proxy.ReadSyncState<FNetPredictionTestSyncState>();
+			CHECK(ClientState1->AutoCounter == ExpectedServerAutoCounterValue - ExpectedInterpolateBufferFrames);
+			CHECK(ClientState1->InputCounter == 0.0f);
+
+			// Object2 is now an autonomous proxy
+			ClientState2 = Worlds.Object2.ClientObject.Proxy.ReadSyncState<FNetPredictionTestSyncState>();
+			if (i == 0)
+			{
+				// For the first client tick after "possession" the client hasn't run forward yet
+				CHECK(ClientState2->AutoCounter == ExpectedServerAutoCounterValue);
+			}
+			else
+			{
+				// Now the client is a few frames ahead of the server
+				CHECK(ClientState2->AutoCounter == ExpectedServerAutoCounterValue + ExpectedForwardPredictFrames);
+			}
+			CHECK(ClientState2->InputCounter == 0.0f);
+		}
+	}
+}
+
+}
\ No newline at end of file
