Update network prediction tests. Emulate only sending changed properties of the prediction driver, which fixes an issue in the tests where serializing the FNetworkPredictionProxy every tick was causing it to be reconfigured every tick. As a result, frames of the frame data buffer were being overwritten unnecessarily, due to the PendingFrame check in UNetworkPredictionWorldManager::ConfigureInstance.

diff --git a/Engine/Source/Programs/NetworkPredictionTests/Private/NetPredictionTestDriver.cpp b/Engine/Source/Programs/NetworkPredictionTests/Private/NetPredictionTestDriver.cpp
--- a/Engine/Source/Programs/NetworkPredictionTests/Private/NetPredictionTestDriver.cpp
+++ b/Engine/Source/Programs/NetworkPredictionTests/Private/NetPredictionTestDriver.cpp
@@ -1,6 +1,7 @@
 // Copyright Epic Games, Inc. All Rights Reserved.
 
 #include "NetPredictionTestDriver.h"
+#include "HAL/PlatformMath.h"
 #include "NetPredictionMockPackageMap.h"
 #include "NetPredictionTestChannel.h"
 #include "NetworkPredictionModelDefRegistry.h"
@@ -13,6 +14,8 @@
 namespace UE::Net::Private
 {
 
+const uint8 FNetPredictionTestDriver::ChangeFlagsBitsNeeded = 1 + FPlatformMath::FloorLog2(static_cast<uint8>(EChangeFlags::MaxFlagPlusOne) - 1);
+
 struct FTestNetPredictionModelDef : public FNetworkPredictionModelDef
 {
 	NP_MODEL_BODY();
@@ -115,38 +118,31 @@ void FNetPredictionTestDriver::ReceiveServerRPCs()
 	}
 }
 
-void FNetPredictionTestDriver::TestNetSerialize(FArchive& Ar)
+void FNetPredictionTestDriver::NetSerialize(FArchive& Ar, EChangeFlags Flags)
 {
 	bool bSuccess = false;
 
-	Proxy.NetSerialize(Ar, nullptr, bSuccess);
+	Ar.SerializeBits(&Flags, ChangeFlagsBitsNeeded);
 
-	// Emulate the replication conditions from UNetworkPredictionComponent
-	if (Proxy.GetCachedHasNetConnection())
+	if (EnumHasAnyFlags(Flags, EChangeFlags::PredictionProxyChanged))
 	{
-		ReplicationProxy_Autonomous.NetSerialize(Ar, UNetPredictionMockPackageMap::Get(), bSuccess);
+		Proxy.NetSerialize(Ar, UNetPredictionMockPackageMap::Get(), bSuccess);
 	}
-	else
+
+	if (EnumHasAnyFlags(Flags, EChangeFlags::ReplicationProxyChanged))
 	{
-		ReplicationProxy_Simulated.NetSerialize(Ar, UNetPredictionMockPackageMap::Get(), bSuccess);
+		// Emulate the replication conditions from UNetworkPredictionComponent
+		if (Proxy.GetCachedHasNetConnection())
+		{
+			ReplicationProxy_Autonomous.NetSerialize(Ar, UNetPredictionMockPackageMap::Get(), bSuccess);
+		}
+		else
+		{
+			ReplicationProxy_Simulated.NetSerialize(Ar, UNetPredictionMockPackageMap::Get(), bSuccess);
+		}
 	}
 }
 
-void FNetPredictionTestDriver::ServerSend()
-{
-	FNetBitWriter TempWriter(1024 * 8 * 2);
-	TestNetSerialize(TempWriter);
-	ServerToClient->Send(TempWriter);
-}
-
-void FNetPredictionTestDriver::ClientReceive()
-{
-	REQUIRE(ServerToClient->HasPendingData());
-
-	FNetBitReader TempReader = ServerToClient->Receive();
-	TestNetSerialize(TempReader);
-}
-
 FNetPredictionTestObject::FNetPredictionTestObject(
 		UNetworkPredictionWorldManager* ServerWorldManager,
 		UNetworkPredictionWorldManager* ClientWorldManager,
@@ -160,4 +156,51 @@ FNetPredictionTestObject::FNetPredictionTestObject(
 	ClientObject.Proxy.InitForNetworkRole(ClientRole, ClientRole == ENetRole::ROLE_AutonomousProxy);
 }
 
+void FNetPredictionTestObject::ServerSend()
+{
+	ServerObject.ReplicationProxy_Autonomous.OnPreReplication();
+	ServerObject.ReplicationProxy_Simulated.OnPreReplication();
+	
+	// Simple emulation of not sending changed properties.
+	// Note this only works when no "packets" are "dropped" in the tests, which is currently the case
+	FNetPredictionTestDriver::EChangeFlags Flags = FNetPredictionTestDriver::EChangeFlags::None;
+	
+	if (!ProxyShadowState.Identical(&ServerObject.Proxy, 0))
+	{
+		Flags |= FNetPredictionTestDriver::EChangeFlags::PredictionProxyChanged;
+		ProxyShadowState = ServerObject.Proxy;
+	}
+
+	// Emulate the replication conditions from UNetworkPredictionComponent
+	if (ServerObject.Proxy.GetCachedHasNetConnection())
+	{
+		if (!AutonomousProxyShadowState.Identical(&ServerObject.ReplicationProxy_Autonomous, 0))
+		{
+			Flags |= FNetPredictionTestDriver::EChangeFlags::ReplicationProxyChanged;
+			AutonomousProxyShadowState = ServerObject.ReplicationProxy_Autonomous;
+		}
+	}
+	else
+	{
+		if (!SimulatedProxyShadowState.Identical(&ServerObject.ReplicationProxy_Simulated, 0))
+		{
+			Flags |= FNetPredictionTestDriver::EChangeFlags::ReplicationProxyChanged;
+			SimulatedProxyShadowState = ServerObject.ReplicationProxy_Simulated;
+		}
+	}
+
+	FNetBitWriter TempWriter(1024 * 8 * 2);
+	ServerObject.NetSerialize(TempWriter, Flags);
+	ServerToClient->Send(TempWriter);
+}
+
+void FNetPredictionTestObject::ClientReceive()
+{
+	while(ServerToClient->HasPendingData())
+	{
+		FNetBitReader TempReader = ServerToClient->Receive();
+		ClientObject.NetSerialize(TempReader, FNetPredictionTestDriver::EChangeFlags());
+	}
+}
+
 }
\ No newline at end of file
diff --git a/Engine/Source/Programs/NetworkPredictionTests/Private/NetPredictionTestDriver.h b/Engine/Source/Programs/NetworkPredictionTests/Private/NetPredictionTestDriver.h
--- a/Engine/Source/Programs/NetworkPredictionTests/Private/NetPredictionTestDriver.h
+++ b/Engine/Source/Programs/NetworkPredictionTests/Private/NetPredictionTestDriver.h
@@ -3,6 +3,7 @@
 #pragma once
 
 #include "Engine/EngineBaseTypes.h"
+#include "Misc/EnumClassFlags.h"
 #include "NetPredictionTestChannel.h"
 #include "NetworkPredictionModelDef.h"
 #include "NetworkPredictionProxy.h"
@@ -69,6 +70,17 @@ struct FNetPredictionTestDriver
 	bool bHidden = false;
 	bool bInputPressed = false;
 
+	enum class EChangeFlags : uint8
+	{
+		None = 0,
+		PredictionProxyChanged = 1 << 0,
+		ReplicationProxyChanged = 1 << 1,
+
+		// Add new entries above
+		MaxFlagPlusOne
+	};
+	static const uint8 ChangeFlagsBitsNeeded;
+
 	FNetPredictionTestDriver(UNetworkPredictionWorldManager* WorldManager, ENetMode Mode,
 		TSharedPtr<FNetPredictionTestChannel> InClientToServer,
 		TSharedPtr<FNetPredictionTestChannel> InServerToClient);
@@ -81,10 +93,7 @@ struct FNetPredictionTestDriver
 	void CallServerRPC();
 	void ReceiveServerRPCs();
 
-	void TestNetSerialize(FArchive& Ar);
-
-	void ServerSend();
-	void ClientReceive();
+	void NetSerialize(FArchive& Ar, EChangeFlags Flags);
 
 private:
 	FReplicationProxySet GetReplicationProxies()
@@ -96,6 +105,8 @@ private:
 	TSharedPtr<FNetPredictionTestChannel> ServerToClient;
 };
 
+ENUM_CLASS_FLAGS(FNetPredictionTestDriver::EChangeFlags);
+
 // Conveniently associates the server & client object drivers and their channels for ease of writing unit tests.
 struct FNetPredictionTestObject
 {
@@ -109,6 +120,16 @@ struct FNetPredictionTestObject
 		UNetworkPredictionWorldManager* ServerWorldManager,
 		UNetworkPredictionWorldManager* ClientWorldManager,
 		ENetRole ClientRole);
+
+	void ServerSend();
+	void ClientReceive();
+
+private:
+	// Stores previously sent state, compared with current state to prevent "sending" when it hasn't changed.
+	// For emulating FRepLayout behavior.
+	FNetworkPredictionProxy ProxyShadowState;
+	FReplicationProxy AutonomousProxyShadowState;
+	FReplicationProxy SimulatedProxyShadowState;
 };
 
 }
\ No newline at end of file
diff --git a/Engine/Source/Programs/NetworkPredictionTests/Private/WorldManagerTests_ForwardPredict.cpp b/Engine/Source/Programs/NetworkPredictionTests/Private/WorldManagerTests_ForwardPredict.cpp
--- a/Engine/Source/Programs/NetworkPredictionTests/Private/WorldManagerTests_ForwardPredict.cpp
+++ b/Engine/Source/Programs/NetworkPredictionTests/Private/WorldManagerTests_ForwardPredict.cpp
@@ -105,13 +105,13 @@ struct FSingleSimulatedObjectWorld
 		ServerWorld.PreTick(PredictionTestFixedFrameRate);
 		SimObject.ServerObject.ReceiveServerRPCs();
 		ServerWorld.Tick(PredictionTestFixedFrameRate);
-		SimObject.ServerObject.ServerSend();
+		SimObject.ServerSend();
 	}
 
 	void TickClient()
 	{
 		ClientWorld.PreTick(PredictionTestFixedFrameRate);
-		SimObject.ClientObject.ClientReceive();
+		SimObject.ClientReceive();
 		ClientWorld.Tick(PredictionTestFixedFrameRate);
 	}
 };
diff --git a/Engine/Source/Programs/NetworkPredictionTests/Private/WorldManagerTests_Interpolation.cpp b/Engine/Source/Programs/NetworkPredictionTests/Private/WorldManagerTests_Interpolation.cpp
--- a/Engine/Source/Programs/NetworkPredictionTests/Private/WorldManagerTests_Interpolation.cpp
+++ b/Engine/Source/Programs/NetworkPredictionTests/Private/WorldManagerTests_Interpolation.cpp
@@ -32,15 +32,15 @@ struct FTestPossession
 		Object1.ServerObject.ReceiveServerRPCs();
 		Object2.ServerObject.ReceiveServerRPCs();
 		ServerWorld.Tick(FixedFrameRate);
-		Object1.ServerObject.ServerSend();
-		Object2.ServerObject.ServerSend();
+		Object1.ServerSend();
+		Object2.ServerSend();
 	}
 
 	void TickClient()
 	{
 		ClientWorld.PreTick(FixedFrameRate);
-		Object1.ClientObject.ClientReceive();
-		Object2.ClientObject.ClientReceive();
+		Object1.ClientReceive();
+		Object2.ClientReceive();
 		ClientWorld.Tick(FixedFrameRate);
 	}
 
@@ -55,6 +55,11 @@ TEST_CASE("NetworkPrediction interpolation mode, one simulated proxy", "[.][poss
 	Settings->Settings.SimulatedProxyNetworkLOD = ENetworkLOD::Interpolated;
 
 	FTestPossession Worlds(Settings);
+		
+	Worlds.Object1.ServerObject.DebugName = TEXT("Obj1Server");
+	Worlds.Object1.ClientObject.DebugName = TEXT("Obj1Client");
+	Worlds.Object2.ServerObject.DebugName = TEXT("Obj2Server");
+	Worlds.Object2.ClientObject.DebugName = TEXT("Obj2Client");
 
 	float ExpectedServerAutoCounterValue = 0.0f;
 
@@ -107,11 +112,6 @@ TEST_CASE("NetworkPrediction interpolation mode, one simulated proxy", "[.][poss
 		const FNetPredictionTestSyncState* ServerState1 = nullptr;
 		const FNetPredictionTestSyncState* ServerState2 = nullptr;
 
-		Worlds.Object1.ServerObject.DebugName = TEXT("Obj1Server");
-		Worlds.Object1.ClientObject.DebugName = TEXT("Obj1Client");
-		Worlds.Object2.ServerObject.DebugName = TEXT("Obj2Server");
-		Worlds.Object2.ClientObject.DebugName = TEXT("Obj2Client");
-
 		constexpr int32 ExpectedForwardPredictFrames = 6;
 
 		// Run some frames to allow the client to predict ahead (autonomous proxy)
