Fix modules with mismatched platform support attributes in corresponding .Target.cs and .Build.cs files.

diff --git a/Engine/Source/Programs/NetworkPredictionTests/NetworkPredictionTests.Target.cs b/Engine/Source/Programs/NetworkPredictionTests/NetworkPredictionTests.Target.cs
--- a/Engine/Source/Programs/NetworkPredictionTests/NetworkPredictionTests.Target.cs
+++ b/Engine/Source/Programs/NetworkPredictionTests/NetworkPredictionTests.Target.cs
@@ -2,6 +2,7 @@
 
 using UnrealBuildTool;
 
+[SupportedPlatforms("Win64", "Linux")]
 public class NetworkPredictionTestsTarget : TestTargetRules
 {
 	public NetworkPredictionTestsTarget(TargetInfo Target) : base(Target)
