Suppress errors & warnings in replication tests on Linux.

#rnx

diff --git a/Engine/Source/Programs/NetworkPredictionTests/Private/NetworkPredictionTests.cpp b/Engine/Source/Programs/NetworkPredictionTests/Private/NetworkPredictionTests.cpp
--- a/Engine/Source/Programs/NetworkPredictionTests/Private/NetworkPredictionTests.cpp
+++ b/Engine/Source/Programs/NetworkPredictionTests/Private/NetworkPredictionTests.cpp
@@ -1,6 +1,8 @@
 // Copyright Epic Games, Inc. All Rights Reserved.
 
+#include "HAL/PlatformMisc.h"
 #include "Logging/LogScopedVerbosityOverride.h"
+#include "Misc/CommandLine.h"
 #include "Modules/ModuleManager.h"
 #include "NetworkPredictionTests.h"
 #include "SlateGlobals.h"
@@ -13,11 +15,18 @@ DEFINE_LOG_CATEGORY(LogNetworkPredictionTests);
 
 GROUP_BEFORE_GLOBAL(Catch::DefaultGroup)
 {
+	// Add command-line options always needed for the tests
+	FCommandLine::Append(TEXT(" -nullrhi -unattended"));
+
 	{
+		// Find the engine directory before InitAll() enables the platform file stub. Prevents warnings.
+		FPlatformMisc::EngineDir();
+
 		// Silence some warnings during initialization unrelated to the tests
 		LOG_SCOPE_VERBOSITY_OVERRIDE(LogSlate, ELogVerbosity::Error);
 		LOG_SCOPE_VERBOSITY_OVERRIDE(LogSlateStyle, ELogVerbosity::Error);
-		LOG_SCOPE_VERBOSITY_OVERRIDE(LogUObjectGlobals, ELogVerbosity::Error);
+		LOG_SCOPE_VERBOSITY_OVERRIDE(LogUObjectGlobals, ELogVerbosity::Fatal);
+		LOG_SCOPE_VERBOSITY_OVERRIDE(LogStreaming, ELogVerbosity::Error);
 		InitAll(true, true);
 	}
 
diff --git a/Engine/Source/Programs/ReplicationSystemLowLevelTests/Private/ReplicationSystemLowLevelTests.cpp b/Engine/Source/Programs/ReplicationSystemLowLevelTests/Private/ReplicationSystemLowLevelTests.cpp
--- a/Engine/Source/Programs/ReplicationSystemLowLevelTests/Private/ReplicationSystemLowLevelTests.cpp
+++ b/Engine/Source/Programs/ReplicationSystemLowLevelTests/Private/ReplicationSystemLowLevelTests.cpp
@@ -1,5 +1,6 @@
 // Copyright Epic Games, Inc. All Rights Reserved.
 
+#include "HAL/PlatformMisc.h"
 #include "Iris/IrisConfig.h"
 #include "Logging/LogScopedVerbosityOverride.h"
 #include "Misc/CommandLine.h"
@@ -57,10 +58,14 @@ GROUP_BEFORE_GLOBAL(Catch::DefaultGroup)
 	UE::Net::SetUseIrisReplication(true);
 
 	{
+		// Find the engine directory before InitAll() enables the platform file stub. Prevents warnings.
+		FPlatformMisc::EngineDir();
+
 		// Silence some errors & warnings during initialization unrelated to the replication system
 		LOG_SCOPE_VERBOSITY_OVERRIDE(LogSlate, ELogVerbosity::Error);
 		LOG_SCOPE_VERBOSITY_OVERRIDE(LogSlateStyle, ELogVerbosity::Error);
 		LOG_SCOPE_VERBOSITY_OVERRIDE(LogUObjectGlobals, ELogVerbosity::Fatal);
+		LOG_SCOPE_VERBOSITY_OVERRIDE(LogStreaming, ELogVerbosity::Error);
 		InitAll(true, true);
 	}
 
