Fix NetworkPredictionTests compile errors.

#rnx

diff --git a/Engine/Source/Programs/NetworkPredictionTests/NetworkPredictionTests.Target.cs b/Engine/Source/Programs/NetworkPredictionTests/NetworkPredictionTests.Target.cs
--- a/Engine/Source/Programs/NetworkPredictionTests/NetworkPredictionTests.Target.cs
+++ b/Engine/Source/Programs/NetworkPredictionTests/NetworkPredictionTests.Target.cs
@@ -7,13 +7,13 @@ public class NetworkPredictionTestsTarget : TestTargetRules
 {
 	public NetworkPredictionTestsTarget(TargetInfo Target) : base(Target)
 	{
-		bCompileAgainstEngine = true;
-		bCompileAgainstApplicationCore = true;
-		bUsesSlate = false;
-
+		bTestsRequireEngine = true;
+		bTestsRequireApplicationCore = true;
 		bUsePlatformFileStub = true;
 		bMockEngineDefaults = true;
 
+		bUsesSlate = false;
+
 		IncludeOrderVersion = EngineIncludeOrderVersion.Latest;
 	}
 }
