Fix Network Prediction Test Driver compiler error resulting from recent NPP/Iris support changes.

#Synced-CL 48900023

diff --git a/Engine/Source/Programs/NetworkPredictionTests/Private/NetPredictionTestDriver.cpp b/Engine/Source/Programs/NetworkPredictionTests/Private/NetPredictionTestDriver.cpp
--- a/Engine/Source/Programs/NetworkPredictionTests/Private/NetPredictionTestDriver.cpp
+++ b/Engine/Source/Programs/NetworkPredictionTests/Private/NetPredictionTestDriver.cpp
@@ -100,7 +100,7 @@ void FNetPredictionTestDriver::CallServerRPC()
 
 	FNetBitWriter Ar(1024 * 8 * 2);
 	bool bOutSuccess = false;
-	ProxyParameter.NetSerialize(Ar, UNetPredictionMockPackageMap::Get(), bOutSuccess);
+	ProxyParameter.SerializeToProxy(Ar, UNetPredictionMockPackageMap::Get(), bOutSuccess);
 
 	ClientToServer->Send(Ar);
 }
@@ -111,10 +111,8 @@ void FNetPredictionTestDriver::ReceiveServerRPCs()
 	{
 		FNetBitReader Ar = ClientToServer->Receive();
 
-		FServerReplicationRPCParameter ProxyParameter;
 		bool bOutSuccess = false;
-		ProxyParameter.NetSerialize(Ar, UNetPredictionMockPackageMap::Get(), bOutSuccess);
-		ProxyParameter.NetSerializeToProxy(ReplicationProxy_ServerRPC);
+		ReplicationProxy_ServerRPC.NetSerialize(Ar, UNetPredictionMockPackageMap::Get(), bOutSuccess);
 	}
 }
 
