Network Prediction Plugin Tests:  disabling new time dilation feature for a couple of tests that are relying on strict equal time for both client and server

#rnx

#Synced-CL 51257655

diff --git a/Engine/Source/Programs/NetworkPredictionTests/Private/WorldManagerTests_ForwardPredict.cpp b/Engine/Source/Programs/NetworkPredictionTests/Private/WorldManagerTests_ForwardPredict.cpp
--- a/Engine/Source/Programs/NetworkPredictionTests/Private/WorldManagerTests_ForwardPredict.cpp
+++ b/Engine/Source/Programs/NetworkPredictionTests/Private/WorldManagerTests_ForwardPredict.cpp
@@ -3,6 +3,7 @@
 #include "NetPredictionTestDriver.h"
 #include "NetPredictionTestWorld.h"
 #include "NetworkPredictionDriver.h"
+#include "NetworkPredictionSerialization.h"
 #include "NetworkPredictionProxy.h"
 
 #include <catch2/catch_test_macros.hpp>
@@ -149,6 +150,11 @@ TEST_CASE("NetworkPrediction, one sim proxy object updates sim state on client",
 
 TEST_CASE("NetworkPrediction, one autonomous proxy object updates sim state on client & server")
 {
+	// Disabling time dilation for this test since it dynamically affects how far ahead the client gets. Restored at end.
+	const bool bDefaultEnableTimeDilation = NetworkPredictionCVars::EnableTimeDilation->GetBool();
+	NetworkPredictionCVars::EnableTimeDilation->Set(false);
+	ON_SCOPE_EXIT { NetworkPredictionCVars::EnableTimeDilation->Set(bDefaultEnableTimeDilation); };
+
 	constexpr ENetRole ClientRole = ENetRole::ROLE_AutonomousProxy;
 	FSingleSimulatedObjectWorld TestWorld(ENetRole::ROLE_AutonomousProxy);
 
@@ -219,6 +225,7 @@ TEST_CASE("NetworkPrediction, one autonomous proxy object updates sim state on c
 			CHECK(ClientState->InputCounter == ExpectedClientInputCounterValue);
 		}
 	}
+
 }
 
 }
\ No newline at end of file
diff --git a/Engine/Source/Programs/NetworkPredictionTests/Private/WorldManagerTests_Interpolation.cpp b/Engine/Source/Programs/NetworkPredictionTests/Private/WorldManagerTests_Interpolation.cpp
--- a/Engine/Source/Programs/NetworkPredictionTests/Private/WorldManagerTests_Interpolation.cpp
+++ b/Engine/Source/Programs/NetworkPredictionTests/Private/WorldManagerTests_Interpolation.cpp
@@ -4,6 +4,7 @@
 #include "NetPredictionTestWorld.h"
 #include "NetworkPredictionDriver.h"
 #include "NetworkPredictionProxy.h"
+#include "NetworkPredictionSerialization.h"
 
 #include <catch2/catch_test_macros.hpp>
 
@@ -51,6 +52,11 @@ private:
 // This case includes a repro for UE-170936, disabled until it's fixed.
 TEST_CASE("NetworkPrediction interpolation mode, one simulated proxy", "[possession]")
 {
+	// Disabling time dilation for this test since it dynamically affects how far ahead the client gets. Restored at end.
+	const bool bDefaultEnableTimeDilation = NetworkPredictionCVars::EnableTimeDilation->GetBool();
+	NetworkPredictionCVars::EnableTimeDilation->Set(false);
+	ON_SCOPE_EXIT{ NetworkPredictionCVars::EnableTimeDilation->Set(bDefaultEnableTimeDilation); };
+
 	UNetworkPredictionSettingsObject* Settings = NewObject<UNetworkPredictionSettingsObject>();
 	Settings->Settings.SimulatedProxyNetworkLOD = ENetworkLOD::Interpolated;
 
